Add buoyancy subsystem plugin & initial prototype for midphase-buoyancy

diff --git a/Engine/Plugins/Experimental/Buoyancy/Buoyancy.uplugin b/Engine/Plugins/Experimental/Buoyancy/Buoyancy.uplugin
new file mode 100644
--- /dev/null
+++ b/Engine/Plugins/Experimental/Buoyancy/Buoyancy.uplugin
@@ -0,0 +1,30 @@
+{
+	"FileVersion": 1,
+	"Version": 1,
+	"VersionName": "0.1",
+	"FriendlyName": "Buoyancy",
+	"Description": "Set of tools for adding buoyant behavior to objects which interact with water bodies generated by the Water plugin",
+	"Category": "Water",
+	"CreatedBy": "Epic Games, Inc.",
+	"CreatedByURL": "https://epicgames.com",
+	"DocsURL": "",
+	"MarketplaceURL": "",
+	"SupportURL": "",
+	"EnabledByDefault": false,
+	"CanContainContent": true,
+	"IsExperimentalVersion": true,
+	"Installed": false,
+	"Modules": [
+		{
+			"Name": "Buoyancy",
+			"Type": "Runtime",
+			"LoadingPhase": "PostConfigInit"
+		}
+	],
+	"Plugins": [
+		{
+			"Name": "Water",
+			"Enabled": true
+		},
+	]
+}
\ No newline at end of file
diff --git a/Engine/Plugins/Experimental/Buoyancy/Source/Runtime/Buoyancy.Build.cs b/Engine/Plugins/Experimental/Buoyancy/Source/Runtime/Buoyancy.Build.cs
new file mode 100644
--- /dev/null
+++ b/Engine/Plugins/Experimental/Buoyancy/Source/Runtime/Buoyancy.Build.cs
@@ -0,0 +1,28 @@
+// Copyright Epic Games, Inc. All Rights Reserved.
+
+using UnrealBuildTool;
+
+public class Buoyancy : ModuleRules
+{
+	public Buoyancy(ReadOnlyTargetRules Target) : base(Target)
+	{
+		PCHUsage = ModuleRules.PCHUsageMode.UseExplicitOrSharedPCHs;
+
+		PublicDependencyModuleNames.AddRange(
+			new string[] {
+				"Core",
+				"CoreUObject",
+				"Engine",
+				"Water",
+			}
+		);
+
+		PrivateDependencyModuleNames.AddRange(
+			new string[] {
+				"ChaosCore",
+				"Chaos",
+				"PhysicsCore",
+			}
+		);
+	}
+}
\ No newline at end of file
diff --git a/Engine/Plugins/Experimental/Buoyancy/Source/Runtime/Private/BuoyancyAlgorithms.cpp b/Engine/Plugins/Experimental/Buoyancy/Source/Runtime/Private/BuoyancyAlgorithms.cpp
new file mode 100644
--- /dev/null
+++ b/Engine/Plugins/Experimental/Buoyancy/Source/Runtime/Private/BuoyancyAlgorithms.cpp
@@ -0,0 +1,430 @@
+// Copyright Epic Games, Inc. All Rights Reserved.
+
+#include "BuoyancyAlgorithms.h"
+#include "BuoyancyStats.h"
+#include "DrawDebugHelpers.h"
+#include "Engine/Engine.h"
+#include "Chaos/ParticleHandle.h"
+#include "Chaos/DebugDrawQueue.h"
+#include "Chaos/Utilities.h"
+
+//
+// CVars
+//
+
+extern bool bBuoyancyDebugDraw;
+
+
+//
+// Internal Functions
+//
+
+namespace
+{
+	using namespace Chaos;
+
+	// NOTE: See SubdivideBounds(...) for a description of this algorithm
+	//
+	// TODO: Use an array + offset rather than raw ptr
+	// TODO: Use a bounds stack instead of a recursive algo
+	// TODO: Only build overlapping parts of hierarchy
+	void SubdivideBounds_Internal(const FAABB3& Bounds, int32 NumSubdivisions, FAABB3** BoundsPtrPtr)
+	{
+		// We assume that the buffer pointed to by *BoundsPtr 
+		const FVec3 Min = Bounds.Min();
+		const FVec3 Max = Bounds.Max();
+		const FVec3 Cen = Bounds.GetCenter();
+
+		// Decrement subdivisions and track whether we have any more to go
+		const bool bSubdivide = --NumSubdivisions > 0;
+
+		// Get a pointer to the first writable element in the array of bounds,
+		// or to a temporary swap space of bounds if we haven't reached the
+		// leaf level yet
+		FAABB3 BoundsSwap[8];
+		FAABB3* BoundsPtr = bSubdivide ? BoundsSwap : *BoundsPtrPtr;
+
+		// Generate 8 subdivisions
+		BoundsPtr[0] = FAABB3(Min, Cen);
+		BoundsPtr[1] = FAABB3(FVec3(Cen.X, Min.Y, Min.Z), FVec3(Max.X, Cen.Y, Cen.Z));
+		BoundsPtr[2] = FAABB3(FVec3(Min.X, Cen.Y, Min.Z), FVec3(Cen.X, Max.Y, Cen.Z));
+		BoundsPtr[3] = FAABB3(FVec3(Min.X, Min.Y, Cen.Z), FVec3(Cen.X, Cen.Y, Max.Z));
+		BoundsPtr[4] = FAABB3(Cen, Max);
+		BoundsPtr[5] = FAABB3(FVec3(Min.X, Cen.Y, Cen.Z), FVec3(Cen.X, Max.Y, Max.Z));
+		BoundsPtr[6] = FAABB3(FVec3(Cen.X, Min.Y, Cen.Z), FVec3(Max.X, Cen.Y, Max.Z));
+		BoundsPtr[7] = FAABB3(FVec3(Cen.X, Cen.Y, Min.Z), FVec3(Max.X, Max.Y, Cen.Z));
+
+		if (bSubdivide)
+		{
+			// Recurse if we haven't reached the leaf level yet
+			for (int32 Index = 0; Index < 8; ++Index)
+			{
+				SubdivideBounds_Internal(BoundsPtr[Index], NumSubdivisions, BoundsPtrPtr);
+			}
+		}
+		else
+		{
+			// Increment the bounds ptr by 8 if we just wrote 8 leaves
+			(*BoundsPtrPtr) += 8;
+		}
+	}
+}
+
+
+//
+// Algorithms
+//
+
+namespace BuoyancyAlgorithms
+{
+	using namespace Chaos;
+
+	bool ComputeSubmergedVolume(const FGeometryParticleHandle* ParticleA, const FGeometryParticleHandle* ParticleB, const int32 NumSubdivisions, const float MinVolume, float& SubmergedVol, FVec3& SubmergedCoM)
+	{
+		SCOPE_CYCLE_COUNTER(STAT_BuoyancyAlgorithms_ComputeSubmergedVolume)
+
+		SubmergedVol = 0.f;
+		SubmergedCoM = FVec3::ZeroVector;
+
+		const FImplicitObject* RootImplicitA = ParticleA->GetGeometry();
+		const FImplicitObject* RootImplicitB = ParticleB->GetGeometry();
+
+		const FConstGenericParticleHandle PA = ParticleA;
+		const FConstGenericParticleHandle PB = ParticleB;
+		const FShapeInstanceArray& ShapeInstancesA = ParticleA->ShapeInstances();
+		const FShapeInstanceArray& ShapeInstancesB = ParticleB->ShapeInstances();
+
+		// Particle transforms
+		const FRigidTransform3 ParticleWorldTransformA = PA->GetTransformPQ();
+		const FRigidTransform3 ParticleWorldTransformB = PB->GetTransformPQ();
+		const FRigidTransform3 ParticleTransformAToB = ParticleWorldTransformA.GetRelativeTransform(ParticleWorldTransformB);
+
+		// Detect collisions between Implicit Hierarchy of ParticleA and Implicit Hierarchy of
+		// ParticleB. Given an ImplicitObject from ParticleA (which we know overlaps the bounds
+		// of some parts of ParticleB), run collision detection on ImplicitA against the implicit
+		// object hierarchy of ParticleB.
+		RootImplicitA->VisitLeafObjects(
+			[ParticleA, &ShapeInstancesA, &ParticleWorldTransformA,
+			ParticleB, &ShapeInstancesB, RootImplicitB, &ParticleWorldTransformB, &ParticleTransformAToB,
+			&NumSubdivisions, &MinVolume, &SubmergedVol, &SubmergedCoM]
+			(const FImplicitObject* ImplicitA, const FRigidTransform3& RelativeTransformA, const int32 RootObjectIndexA, const int32 ObjectIndex, const int32 LeafObjectIndexA)
+		{
+			const FAABB3 RelativeBoundsA = ImplicitA->CalculateTransformedBounds(RelativeTransformA);
+			const FAABB3 ShapeBoundsAInB = RelativeBoundsA.TransformedAABB(ParticleTransformAToB);
+			const int32 ShapeIndexA = (ShapeInstancesA.IsValidIndex(RootObjectIndexA)) ? RootObjectIndexA : 0;
+			const FShapeInstance* ShapeInstanceA = ShapeInstancesA[ShapeIndexA].Get();
+
+			//
+			// TODO: Dig down into ImplicitB's leaves based on its type, so that
+			// we don't end up going down to the tiniest particles of a GC.
+			//
+
+			// Detect collisions between ImplicitA and Implicit Hierarchy of ParticleB
+			RootImplicitB->VisitOverlappingLeafObjects(ShapeBoundsAInB,
+				[ParticleA, ImplicitA, ShapeInstanceA, &ParticleWorldTransformA, &RelativeTransformA, LeafObjectIndexA,
+				ParticleB, &ParticleWorldTransformB, &ShapeInstancesB,
+				&NumSubdivisions, &MinVolume, &SubmergedVol, &SubmergedCoM]
+				(const FImplicitObject* ImplicitB, const FRigidTransform3& RelativeTransformB, const int32 RootObjectIndexB, const int32 ObjectIndexB, const int32 LeafObjectIndexB)
+			{
+				// Get shape world transforms
+				const FRigidTransform3 ShapeWorldTransformA = RelativeTransformA * ParticleWorldTransformA;
+				const FRigidTransform3 ShapeWorldTransformB = RelativeTransformB * ParticleWorldTransformB;
+
+				// Get local bounds of objects
+				const FAABB3 BoxAInA = ImplicitA->BoundingBox();
+				const FAABB3 BoxBInB = ImplicitB->BoundingBox();
+
+				// AABB intersect check
+				const FAABB3 ShapeWorldBoundsA = BoxAInA.TransformedAABB(ShapeWorldTransformA);
+				const FAABB3 ShapeWorldBoundsB = BoxBInB.TransformedAABB(ShapeWorldTransformB);
+				if (!ShapeWorldBoundsA.Intersects(ShapeWorldBoundsB)) { return; }
+
+				// OOBB vs AABB intersect checks
+				const FRigidTransform3 ShapeTransformBToA = ShapeWorldTransformB.GetRelativeTransform(ShapeWorldTransformA);
+				const FAABB3 BoxBInA = BoxBInB.TransformedAABB(ShapeTransformBToA);
+				if (!ImplicitA->BoundingBox().Intersects(BoxBInA)) { return; }
+				const FAABB3 BoxAInB = BoxAInA.TransformedAABB(ShapeTransformBToA.Inverse());
+				if (!ImplicitB->BoundingBox().Intersects(BoxAInB)) { return; }
+
+				// Generate subdivided bounds list
+				TArray<FAABB3> BoxesInB;
+				SubdivideBounds(BoxBInB, NumSubdivisions, MinVolume, BoxesInB);
+
+				// Loop over every subdivision of the shape bounds, counting up submerged portions
+				for (const FAABB3& BoxInB : BoxesInB)
+				{
+#if ENABLE_DRAW_DEBUG
+					if (bBuoyancyDebugDraw)
+					{
+						// Subdivided Rigid OOBB
+						Chaos::FDebugDrawQueue::GetInstance().DrawDebugBox(
+							ShapeWorldTransformB.TransformPosition(BoxInB.GetCenter()),
+							BoxInB.Extents() * .5f,
+							ShapeWorldTransformB.GetRotation(),
+							FColor::Green, false, -1.f, SDPG_Foreground, 1.f);
+					}
+#endif
+
+					// Compute approximate submerged bounds & update submersion
+					const float WaterZ = ShapeWorldBoundsA.Max().Z;
+					FAABB3 SubmergedBoundsInB;
+					if (ComputeSubmergedBounds(WaterZ, BoxInB, ShapeWorldTransformB, SubmergedBoundsInB))
+					{
+						// This bounds box is submerged. Compute it's volume and center of mass
+						// in world space, and add those contributions to the submerged quantity.
+						const FVec3 LeafSubmergedCoM = ShapeWorldTransformB.TransformPosition(SubmergedBoundsInB.GetCenter());
+						const float LeafSubmergedVol = SubmergedBoundsInB.GetVolume();
+						SubmergedVol += LeafSubmergedVol;
+						SubmergedCoM += LeafSubmergedCoM * LeafSubmergedVol;
+
+						// Make sure the volume of the submerged portion never exceeds the total
+						// volume of the leaf bounds
+						const float LeafMaxVol = BoxBInB.GetVolume() + UE_SMALL_NUMBER;
+						ensureAlwaysMsgf(LeafSubmergedVol <= LeafMaxVol, TEXT("BuoyancyAlgorithms::ComputeSubmergedVolume: The volume of the submerged portion of the leaf bounds has somehow exceeded the volume of the overall leaf bounds."));
+
+#if ENABLE_DRAW_DEBUG
+						if (bBuoyancyDebugDraw)
+						{
+							// Submerged Bounds
+							Chaos::FDebugDrawQueue::GetInstance().DrawDebugBox(
+								ShapeWorldTransformB.TransformPosition(SubmergedBoundsInB.GetCenter()),
+								SubmergedBoundsInB.Extents() * .5f,
+								ShapeWorldTransformB.GetRotation(),
+								FColor::Orange, false, -1.f, SDPG_Foreground, 2.f);
+						}
+#endif
+					}
+				}
+
+#if ENABLE_DRAW_DEBUG
+				if (bBuoyancyDebugDraw)
+				{
+					// Water OOBB
+					Chaos::FDebugDrawQueue::GetInstance().DrawDebugBox(
+						ShapeWorldTransformA.TransformPosition(BoxAInA.GetCenter()),
+						BoxAInA.Extents() * .5f,
+						ShapeWorldTransformA.GetRotation(),
+						FColor::Cyan, false, -1.f, SDPG_Foreground, 2.f);
+
+					// Rigid OOBB
+					Chaos::FDebugDrawQueue::GetInstance().DrawDebugBox(
+						ShapeWorldTransformB.TransformPosition(BoxBInB.GetCenter()),
+						BoxBInB.Extents() * .5f,
+						ShapeWorldTransformB.GetRotation(),
+						FColor::Green, false, -1.f, SDPG_Foreground, 2.f);
+				}
+#endif
+			});
+		});
+
+		// We've accumulated submersion data for all submerged leaves, now we need to
+		// normalize the submerged center of mass
+		if (SubmergedVol > SMALL_NUMBER)
+		{
+			SubmergedCoM /= SubmergedVol;
+
+#if ENABLE_DRAW_DEBUG
+			if (bBuoyancyDebugDraw)
+			{
+				// Draw a point at the submerged CoM
+				Chaos::FDebugDrawQueue::GetInstance().DrawDebugPoint(SubmergedCoM, FColor::Orange, false, -1.f, SDPG_Foreground, 5.f);
+			}
+#endif
+
+			return true;
+		}
+		else
+		{
+			return false;
+		}
+	}
+
+
+	bool ComputeSubmergedBounds(float WaterZ, const FAABB3& RigidBox, const FRigidTransform3& RigidTransform, FAABB3& OutSubmergedBounds)
+	{
+		SCOPE_CYCLE_COUNTER(STAT_BuoyancyAlgorithms_ComputeSubmergedBounds)
+
+		// Get a point and a normal direction on the water representing the surface
+		// of the water in the space of the rigid object.
+		const FVec3 SurfaceNormal = RigidTransform.InverseTransformVector(FVec3::UpVector);
+		const FVec3 SurfacePoint = RigidTransform.InverseTransformPosition(FVec3::UpVector * WaterZ);
+
+		// Partly submerged object can have at most 10 points intersecting
+		// with the water surface
+		TArray<FVec3, TInlineAllocator<10>> SubmergedVertices;
+
+		// Find bound box indices that are submerged
+		for (int32 VertexIndex = 0; VertexIndex < 8; ++VertexIndex)
+		{
+			const FVec3 Vertex = RigidBox.GetVertex(VertexIndex);
+			const float Depth = SurfaceNormal.Dot(SurfacePoint - Vertex);
+			if (Depth > 0.f)
+			{
+				SubmergedVertices.Add(Vertex);
+			}
+		}
+
+		// If no box corners were submerged, then there can be no submerged edges so stop here
+		if (SubmergedVertices.Num() == 0)
+		{
+			return false;
+		}
+
+		// Find intersections of AABB edges with the surface and add these
+		// points to the submerged verts list
+		TArray<FVec3, TInlineAllocator<6>> SurfaceVertices;
+		for (int32 EdgeIndex = 0; EdgeIndex < 12; ++EdgeIndex)
+		{
+			const FAABBEdge Edge = RigidBox.GetEdge(EdgeIndex);
+			const FVec3& Vert0 = RigidBox.GetVertex(Edge.VertexIndex0);
+			const FVec3& Vert1 = RigidBox.GetVertex(Edge.VertexIndex1);
+			const float Depth0 = SurfaceNormal.Dot(SurfacePoint - Vert0);
+			const float Depth1 = SurfaceNormal.Dot(SurfacePoint - Vert1);
+			const bool bSubmerged0 = (Depth0 > 0.f);
+			const bool bSubmerged1 = (Depth1 > 0.f);
+			if (bSubmerged0 ^ bSubmerged1)
+			{
+				const float DepthDiff = Depth0 - Depth1;
+				const float DepthAlpha = Depth0 / DepthDiff; // NOTE: Since one is submerged and one is not, we know that |DepthDiff| > 0
+				const FVec3 SurfaceVertex = FMath::Lerp(Vert0, Vert1, DepthAlpha);
+				SubmergedVertices.Add(SurfaceVertex);
+			}
+		}
+
+		// If we didn't have any submerged vertices, stop here
+		if (SubmergedVertices.Num() == 0)
+		{
+			return false;
+		}
+
+#if ENABLE_DRAW_DEBUG
+		if (bBuoyancyDebugDraw)
+		{
+			// Draw a point at each submerged vertex
+			for (const FVec3& SubmergedVertex : SubmergedVertices)
+			{
+				Chaos::FDebugDrawQueue::GetInstance().DrawDebugPoint(RigidTransform.TransformPosition(SubmergedVertex), FColor::Magenta, false, -1.f, SDPG_Foreground, 3.f);
+			}
+		}
+#endif
+
+		// Build and return an AABB which contains the submerged vertices of the rigid bounds
+		OutSubmergedBounds = FAABB3(SubmergedVertices[0], SubmergedVertices[0]);
+		for (int32 VertexIndex = 1; VertexIndex < SubmergedVertices.Num(); ++VertexIndex)
+		{
+			OutSubmergedBounds.GrowToInclude(SubmergedVertices[VertexIndex]);
+		}
+		return true;
+	}
+
+
+	bool SubdivideBounds(const FAABB3& Bounds, int32 NumSubdivisions, const float MinVolume, TArray<FAABB3>& OutBounds)
+	{
+		SCOPE_CYCLE_COUNTER(STAT_BuoyancyAlgorithms_SubdivideBounds)
+
+		// Initialize bounds to an array of just the original bounds;
+		OutBounds = TArray<FAABB3>{ Bounds };
+
+		// If the bounds volume is already too small to subdivide, return the original bounds only
+		const float Volume = Bounds.GetVolume();
+		if (Volume < SMALL_NUMBER || Volume <= MinVolume)
+		{
+			return false;
+		}
+
+		// If V_0 is the volume of the outermost AABB, then the volume of
+		// one box in the n'th subdivision of an AABB is given by
+		// 
+		// V_n = V_0 * 2^(-3 n)
+		//
+		// We can invert this equation to find the level of subdivisions
+		// at which the volume becomes smaller than V_min.
+		//
+		// n < -(1/3) * log2(V_min / V_0)
+		const int32 MaxNumSubdivisions = int32(-(1.f/3.f) * FMath::Log2(MinVolume / Volume));
+		NumSubdivisions = FMath::Min(MaxNumSubdivisions, NumSubdivisions);
+
+		// If we have any subdivisions to process, do them now
+		if (NumSubdivisions > 0)
+		{
+			// Predetermine the total number of boxes we're going to generate, and allocate them
+			// in a block.
+			const int32 NumBounds = (int32)FMath::Pow(8.f, NumSubdivisions);
+			OutBounds.SetNum(NumBounds);
+			FAABB3* BoundsPtr = OutBounds.GetData();
+
+			// Recursively generate boxes
+			SubdivideBounds_Internal(Bounds, NumSubdivisions, &BoundsPtr);
+
+			// Make sure that we didn't write too many or too few boxes
+			const FAABB3* BoundsPtrMax = OutBounds.GetData() + NumBounds;
+			ensureAlways(BoundsPtr == BoundsPtrMax);
+
+			// Return the box array
+			return true;
+		}
+
+		return false;
+	}
+
+	bool ComputeBuoyantForce(const FPBDRigidParticleHandle* RigidParticle, const float DeltaSeconds, const float WaterDensity, const float WaterDrag, const FVec3& GravityAccelVec, const FVec3& SubmergedCoM, const float SubmergedVol, FVec3& OutDeltaV, FVec3& OutDeltaW)
+	{
+		SCOPE_CYCLE_COUNTER(STAT_BuoyancyAlgorithms_ComputeBuoyantForces)
+
+		// NOTE: We assume gravity is -Z for perf... If we want to support buoyancy for
+		// weird gravity setups, this is where we'd have to fix it up.
+		const FVec3 GravityDir = FVec3::DownVector;
+		const float GravityAccel = FVec3::DotProduct(GravityDir, GravityAccelVec);
+
+		// Compute buoyant force
+		//
+		// NOTE: This is easy to compute with Archimedes' principle
+		// https://en.wikipedia.org/wiki/Buoyancy
+		//
+		const float BuoyantForce = WaterDensity * SubmergedVol * GravityAccel;
+		//                       = [ kg / cm^3 ] * [ cm^3 ]    * [cm / s^2]
+		//                       = [ kg * cm / s^2 ]
+		//                       = [ force ]
+
+		// Only proceed if buoyant force isn't vanishingly small
+		if (BuoyantForce < SMALL_NUMBER)
+		{
+			return false;
+		}
+
+		// Get a generic particle wrapper
+		FConstGenericParticleHandle RigidGeneric(RigidParticle);
+
+		// Get inverse inertia data to compute world space accelerations
+		const FVec3 WorldCoM = RigidGeneric->PCom();
+		const FVec3 CoMDiff = SubmergedCoM - WorldCoM;
+		const FMatrix33 WorldInvI = Utilities::ComputeWorldSpaceInertia(RigidGeneric->RCom(), RigidGeneric->ConditionedInvI());
+
+		// Compute world buoyant force and torque
+		const FVec3 WorldForce = -GravityDir * BuoyantForce;
+		const FVec3 WorldTorque = Chaos::FVec3::CrossProduct(CoMDiff, WorldForce);
+
+		// Use inertia to convert forces to accelerations
+		const FVec3 LinearAccel = RigidGeneric->InvM() * WorldForce;
+		const FVec3 AngularAccel = WorldInvI * WorldTorque;
+
+		// Integrate to get delta velocities
+		OutDeltaV = LinearAccel * DeltaSeconds;
+		OutDeltaW = AngularAccel * DeltaSeconds;
+
+		// Compute water drag force
+		//
+		// NOTE: This is a very approximate "ether drag" style model here, probably 
+		// we should scale the model with submerged volume for more accuracy,
+		// and apply the drag force in opposition to the linear motion of the submerged
+		// center of mass.
+		const float DragFactor = FMath::Max(0.f, 1.f - (WaterDrag * DeltaSeconds));
+
+		// Account for water drag in deltas
+		OutDeltaV = (DragFactor * OutDeltaV) + (DragFactor - 1.f) * RigidParticle->V();
+		OutDeltaW = (DragFactor * OutDeltaW) + (DragFactor - 1.f) * RigidParticle->W();
+
+		//
+		return true;
+	}
+}
diff --git a/Engine/Plugins/Experimental/Buoyancy/Source/Runtime/Private/BuoyancyModule.cpp b/Engine/Plugins/Experimental/Buoyancy/Source/Runtime/Private/BuoyancyModule.cpp
new file mode 100644
--- /dev/null
+++ b/Engine/Plugins/Experimental/Buoyancy/Source/Runtime/Private/BuoyancyModule.cpp
@@ -0,0 +1,6 @@
+// Copyright Epic Games, Inc. All Rights Reserved.
+
+#include "BuoyancyModule.h"
+#include "Modules/ModuleManager.h"
+
+IMPLEMENT_MODULE(FBuoyancyModule, Buoyancy)
diff --git a/Engine/Plugins/Experimental/Buoyancy/Source/Runtime/Private/BuoyancySubsystem.cpp b/Engine/Plugins/Experimental/Buoyancy/Source/Runtime/Private/BuoyancySubsystem.cpp
new file mode 100644
--- /dev/null
+++ b/Engine/Plugins/Experimental/Buoyancy/Source/Runtime/Private/BuoyancySubsystem.cpp
@@ -0,0 +1,320 @@
+// Copyright Epic Games, Inc. All Rights Reserved.
+
+#include "BuoyancySubsystem.h"
+#include "BuoyancyAlgorithms.h"
+#include "BuoyancyStats.h"
+#include "WaterRuntimeSettings.h"
+#include "Engine/World.h"
+#include "Physics/Experimental/PhysScene_Chaos.h"
+#include "Physics/PhysicsFiltering.h"
+#include "Chaos/MidPhaseModification.h"
+#include "Math/UnitConversion.h"
+#include "PBDRigidsSolver.h"			// Only needed to get perparticlegravity :(
+#include "Chaos/PBDRigidsEvolutionGBF.h"// Only needed to get perparticlegravity :(
+#include "Chaos/PerParticleGravity.h"	// Needed in order to determine force of gravity on each particle
+
+//
+// CVars
+//
+
+bool bBuoyancyDebugDraw = false;
+FAutoConsoleVariableRef CVarBuoyancyDebugDraw(TEXT("p.Buoyancy.DebugDraw"), bBuoyancyDebugDraw, TEXT(""));
+
+bool bBuoyancyEnabled = true;
+FAutoConsoleVariableRef CVarBuoyancyEnabled(TEXT("p.Buoyancy.Enabled"), bBuoyancyEnabled, TEXT(""));
+
+
+//
+// Logging
+//
+
+DEFINE_LOG_CATEGORY(LogBuoyancySubsystem);
+
+
+//
+// Buoyancy Subsystem
+//
+
+void UBuoyancySubsystem::PostInitialize()
+{
+	Super::PostInitialize();
+
+	// Create sim callback
+	if (Chaos::FPhysicsSolver* Solver = GetSolver())
+	{
+		SimCallback = Solver->CreateAndRegisterSimCallbackObject_External<FBuoyancySubsystemSimCallback>();
+	}
+
+	// Apply initial runtime settings
+	ApplyRuntimeSettings(GetDefault<UWaterRuntimeSettings>(), EPropertyChangeType::ValueSet);
+
+	// Set up callback for when runtime settings change in editor
+#if WITH_EDITOR
+	GetDefault<UWaterRuntimeSettings>()->OnSettingsChange.AddUObject(this, &UBuoyancySubsystem::ApplyRuntimeSettings);
+#endif //WITH_EDITOR
+}
+
+void UBuoyancySubsystem::Deinitialize()
+{
+	if (SimCallback)
+	{
+		if (Chaos::FPhysicsSolverBase* Solver = SimCallback->GetSolver())
+		{
+			Solver->UnregisterAndFreeSimCallbackObject_External(SimCallback);
+		}
+	}
+
+	Super::Deinitialize();
+}
+
+void UBuoyancySubsystem::ApplyRuntimeSettings(const UWaterRuntimeSettings* InSettings, EPropertyChangeType::Type ChangeType)
+{
+	bBuoyancySettingsChanged = true;
+
+	// Runtime settings presents water density in g/cm^3, but we want it in kg/cm^3
+	// so introduce a factor of 10^-3 here.
+	//BuoyancySettings.WaterDensity = FNumericUnit<EUnit::KilogramsPerCubicCentimeter>::ConvertTo(InSettings->WaterDensity);
+	BuoyancySettings.WaterDensity = FUnitConversion::Convert(InSettings->WaterDensity,
+		EUnit::GramsPerCubicCentimeter,
+		EUnit::KilogramsPerCubicCentimeter);
+
+	BuoyancySettings.WaterCollisionChannel = InSettings->CollisionChannelForWaterObjects;
+}
+
+void UBuoyancySubsystem::Tick(float DeltaTime)
+{
+	//
+	// NOTE: The entire point of this tick is to send runtime data to
+	// the physics thread which might have changed. At the initial
+	// time of this writing, the only piece of data is the water-collision-
+	// channel, which will never change once set in almost every case.
+	// 
+	// Therefore we can probably do away with this tick, and move the
+	// GetProducerInputData_External to PostInitialize and also
+	// ApplyRuntimeSettings for editor builds, and do away with the Tick
+	// entirely.
+	// 
+	// However, I'm not confident that that would work due to timing
+	// weirdness of the sim callback system, and also 
+	//
+
+	Super::Tick(DeltaTime);
+
+	if (bBuoyancyEnabled == false)
+	{
+		return;
+	}
+
+	if (SimCallback == nullptr)
+	{
+		return;
+	}
+
+	// Only bother sending new async inputs if our buoyancy settings actually changed
+	if (bBuoyancySettingsChanged)
+	{
+		if (FBuoyancySubsystemSimCallbackInput* AsyncInput = SimCallback->GetProducerInputData_External())
+		{
+			bBuoyancySettingsChanged = false;
+			AsyncInput->BuoyancySettings = MakeUnique<FBuoyancySettings>(BuoyancySettings);
+		}
+	}
+}
+
+TStatId UBuoyancySubsystem::GetStatId() const
+{
+	RETURN_QUICK_DECLARE_CYCLE_STAT(UBuoyancySubsystem, STATGROUP_Tickables);
+}
+
+Chaos::FPhysicsSolver* UBuoyancySubsystem::GetSolver() const
+{
+	if (UWorld* World = GetWorld())
+	{
+		if (FPhysScene* Scene = World->GetPhysicsScene())
+		{
+			if (Chaos::FPhysicsSolver* Solver = Scene->GetSolver())
+			{
+				return Solver;
+			}
+		}
+	}
+
+	return nullptr;
+}
+
+
+//
+// Buoyancy Sim Callback
+//
+
+void FBuoyancySubsystemSimCallbackInput::Reset()
+{
+	BuoyancySettings.Reset();
+}
+
+namespace
+{
+	// Check to see if any of this particle's shapes has a particular collision channel
+	const bool ParticleHasCollision(const FGeometryParticleHandle& Particle, const ECollisionChannel CollisionChannel)
+	{
+		const auto& Shapes = Particle.ShapesArray();
+		uint32 Word3 = 0;
+		for (const auto& Shape : Shapes)
+		{
+			const FCollisionFilterData& ShapeQueryData = Shape->GetQueryData();
+			Word3 |= ShapeQueryData.Word3;
+		}
+		return GetCollisionChannel(Word3) == CollisionChannel;
+	};
+}
+
+void FBuoyancySubsystemSimCallback::OnMidPhaseModification_Internal(Chaos::FMidPhaseModifierAccessor& MidPhaseAccessor)
+{
+	using namespace Chaos;
+
+	SCOPE_CYCLE_COUNTER(STAT_BuoyancySubsystem_OnMidPhaseModification)
+
+	// If we were sent new buoyancy settings or data, update our local sim copy
+	if (const FBuoyancySubsystemSimCallbackInput* Input = GetConsumerInput_Internal())
+	{
+		if (Input->BuoyancySettings.IsValid())
+		{
+			BuoyancySettings = MoveTemp(Input->BuoyancySettings);
+		}
+	}
+
+	// If we don't have a valid buoyancy settings object, don't continue
+	if (BuoyancySettings.IsValid() == false)
+	{
+		return;
+	}
+
+	// If Buoyancy is disabled, skip
+	if (bBuoyancyEnabled == false)
+	{
+		return;
+	}
+
+	// Get the evolution
+	FPBDRigidsEvolution* Evolution = nullptr;
+	if (FPhysicsSolverBase* SolverBase = GetSolver())
+	{
+		// Why does castchecked return a ref? That makes me think
+		// it's not actually doing a check...
+		FPBDRigidsSolver& PBDSolver = SolverBase->CastChecked();
+		Evolution = PBDSolver.GetEvolution();
+	}
+	if (Evolution == nullptr)
+	{
+		return;
+	}
+
+	// Get perparticle gravity rule, for figuring out the effective gravity on buoyant objects
+	const FPerParticleGravity* PerParticleGravity = &Evolution->GetGravityForces();
+	if (PerParticleGravity == nullptr)
+	{
+		return;
+	}
+
+	// How much time has the sim ticked this frame
+	const FReal DeltaSeconds = GetDeltaTime_Internal();
+
+	// TODO:
+	//
+	// To avoid double-counting any particular object for buoyancy, make a bit array
+	// indexed by particle indices. Every time we apply a buoyant force, we'll skip
+	// that particle next time we see it
+
+	// NOTE: For now we visit _every_ midphase, and check for ones which involve
+	// our target collision channel. It would be nice if it were possible instead
+	// to get a list of water body particles and loop over only midphases which
+	// involve them.
+	MidPhaseAccessor.VisitMidPhases([this, Evolution, PerParticleGravity, DeltaSeconds](Chaos::FMidPhaseModifier& MidPhase)
+	{
+		SCOPE_CYCLE_COUNTER(STAT_BuoyancySubsystem_VisitMidphases)
+
+		FGeometryParticleHandle* Particle0;
+		FGeometryParticleHandle* Particle1;
+		MidPhase.GetParticles(&Particle0, &Particle1);
+		if (Particle0 == nullptr || Particle1 == nullptr)
+		{
+			return;
+		}
+
+		// Use collision filters to pick out which of the particles is water, if either.
+		//
+		// TODO: Instead of this, water particles could have additional userdata, or could
+		// register themselves with the callback and we could look for only the midphases
+		// generated with them.
+		const bool bWater0 = ParticleHasCollision(*Particle0, BuoyancySettings->WaterCollisionChannel);
+		const bool bWater1 = ParticleHasCollision(*Particle1, BuoyancySettings->WaterCollisionChannel);
+		if ((bWater0 || bWater1) == false)
+		{
+			return;
+		}
+
+		// We want to disable all collisions with water, so disable the midphase now
+		MidPhase.Disable();
+
+		// Store the rigid particle and the water particle
+		if (bWater0) { Swap(Particle0, Particle1); }
+		FPBDRigidParticleHandle* RigidParticle = Particle0->CastToRigidParticle();
+		FGeometryParticleHandle* WaterParticle = Particle1;
+		if (RigidParticle == nullptr)
+		{
+			return;
+		}
+
+		// If the particle is in a cluster, get the actual root particle that we can apply forces to
+		FPBDRigidClusteredParticleHandle* RigidCluster = RigidParticle->CastToClustered();
+		if (RigidCluster)
+		{
+			while (FPBDRigidClusteredParticleHandle* ParentCluster = RigidCluster->Parent())
+			{
+				RigidCluster = ParentCluster;
+			}
+		}
+		FPBDRigidParticleHandle* RigidRoot = RigidCluster ? RigidCluster->CastToRigidParticle() : RigidParticle;
+
+		// Compute submerged volume and CoM
+		float SubmergedVol;
+		FVec3 SubmergedCoM;
+		if (BuoyancyAlgorithms::ComputeSubmergedVolume(WaterParticle, RigidParticle, BuoyancySettings->MaxNumBoundsSubdivisions, BuoyancySettings->MinBoundsSubdivisionVol, SubmergedVol, SubmergedCoM))
+		{
+			SCOPE_CYCLE_COUNTER(STAT_BuoyancySubsystem_ComputeBuoyantForces)
+
+			//
+			// If any volume was submerged, take the weighted average of the centers of mass to get the
+			// approximate submerged center of mass, and apply the buoyancy force there.
+			if (SubmergedVol > SMALL_NUMBER)
+			{
+				// Figure out the gravity level of the particle
+				const int32 GravityGroupIndex = RigidParticle->GravityGroupIndex();
+				const FVec3 GravityAccel
+					= PerParticleGravity != nullptr && GravityGroupIndex != INDEX_NONE
+					? (FVec3)PerParticleGravity->GetAcceleration(GravityGroupIndex)
+					: FVec3::DownVector * 980.f; // Default to "regular" gravity
+
+				// Compute delta linear and angular velocities due to buoyancy. If they're big enough to
+				// matter, apply them
+				FVec3 DeltaV, DeltaW;
+				if (BuoyancyAlgorithms::ComputeBuoyantForce(RigidParticle, DeltaSeconds, BuoyancySettings->WaterDensity, BuoyancySettings->WaterDrag, GravityAccel, SubmergedCoM, SubmergedVol, DeltaV, DeltaW))
+				{
+					// Clamp delta velocities
+					DeltaV = DeltaV.GetClampedToSize(0.f, BuoyancySettings->MaxDeltaV);
+					DeltaW = DeltaW.GetClampedToSize(0.f, BuoyancySettings->MaxDeltaW);
+
+					// Apply the deltas
+					RigidParticle->SetV(RigidParticle->V() + DeltaV);
+					RigidParticle->SetW(RigidParticle->W() + DeltaW);
+
+					// Wake up the body??
+					if (Evolution && BuoyancySettings->bKeepAwake)
+					{
+						Evolution->SetParticleObjectState(RigidParticle, EObjectStateType::Dynamic);
+					}
+				}
+			}
+		}
+	});
+}
diff --git a/Engine/Plugins/Experimental/Buoyancy/Source/Runtime/Public/BuoyancyAlgorithms.h b/Engine/Plugins/Experimental/Buoyancy/Source/Runtime/Public/BuoyancyAlgorithms.h
new file mode 100644
--- /dev/null
+++ b/Engine/Plugins/Experimental/Buoyancy/Source/Runtime/Public/BuoyancyAlgorithms.h
@@ -0,0 +1,29 @@
+// Copyright Epic Games, Inc. All Rights Reserved.
+
+#pragma once
+
+#include "CoreMinimal.h"
+#include "Chaos/Core.h"
+#include "Chaos/ParticleHandleFwd.h"
+#include "Chaos/AABB.h"
+
+namespace BuoyancyAlgorithms
+{
+	using namespace Chaos;
+
+	// Compute an approximate volume and center of mass of particle B submerged in particle A
+	bool ComputeSubmergedVolume(const FGeometryParticleHandle* ParticleA, const FGeometryParticleHandle* ParticleB, int32 NumSubdivisions, float MinVolume, float& SubmergedVol, FVec3& SubmergedCoM);
+
+	// Given an OOBB and a water level, generate another OOBB which is 1. entirely contained
+	// within the input OOBB and 2. entirely contains the portion of the OOBB which is submerged
+	// below the water level.
+	bool ComputeSubmergedBounds(float WaterZ, const FAABB3& RigidBox, const FRigidTransform3& RigidTransform, FAABB3& OutSubmergedBounds);
+
+	// Given a bounds object, recursively subdivide it in eighths to a fixed maximum depth and
+	// a fixed minimum smallest subdivision volume.
+	bool SubdivideBounds(const FAABB3& Bounds, int32 NumSubdivisions, float MinVolume, TArray<FAABB3>& OutBounds);
+
+	// Given a rigid particle and it's submerged CoM and Volume, compute delta velocities for
+	// integrated buoyancy forces on an object
+	bool ComputeBuoyantForce(const FPBDRigidParticleHandle* RigidParticle, const float DeltaSeconds, const float WaterDensity, const float WaterDrag, const FVec3& GravityAccelVec, const FVec3& SubmergedCoM, const float SubmergedVol, FVec3& OutDeltaV, FVec3& OutDeltaW);
+}
diff --git a/Engine/Plugins/Experimental/Buoyancy/Source/Runtime/Public/BuoyancyModule.h b/Engine/Plugins/Experimental/Buoyancy/Source/Runtime/Public/BuoyancyModule.h
new file mode 100644
--- /dev/null
+++ b/Engine/Plugins/Experimental/Buoyancy/Source/Runtime/Public/BuoyancyModule.h
@@ -0,0 +1,12 @@
+// Copyright Epic Games, Inc. All Rights Reserved.
+
+#pragma once
+
+#include "Modules/ModuleInterface.h"
+
+class FBuoyancyModule : public IModuleInterface { };
+
+#if UE_ENABLE_INCLUDE_ORDER_DEPRECATED_IN_5_2
+#include "CoreMinimal.h"
+#include "Modules/ModuleManager.h"
+#endif
diff --git a/Engine/Plugins/Experimental/Buoyancy/Source/Runtime/Public/BuoyancyStats.h b/Engine/Plugins/Experimental/Buoyancy/Source/Runtime/Public/BuoyancyStats.h
new file mode 100644
--- /dev/null
+++ b/Engine/Plugins/Experimental/Buoyancy/Source/Runtime/Public/BuoyancyStats.h
@@ -0,0 +1,19 @@
+// Copyright Epic Games, Inc. All Rights Reserved.
+
+#pragma once
+
+#include "Stats/Stats.h"
+
+DECLARE_STATS_GROUP(TEXT("Buoyancy"), STATGROUP_Buoyancy, STATCAT_Advanced);
+
+// Subsystem
+DECLARE_CYCLE_STAT(TEXT("Buoyancy_Subsystem_OnMidPhaseModification"), STAT_BuoyancySubsystem_OnMidPhaseModification, STATGROUP_Buoyancy);
+DECLARE_CYCLE_STAT(TEXT("Buoyancy_Subsystem_VisitMidphases"), STAT_BuoyancySubsystem_VisitMidphases, STATGROUP_Buoyancy);
+DECLARE_CYCLE_STAT(TEXT("Buoyancy_Subsystem_ComputeBuoyantForces"), STAT_BuoyancySubsystem_ComputeBuoyantForces, STATGROUP_Buoyancy);
+
+// Algorithms
+DECLARE_CYCLE_STAT(TEXT("Buoyancy_Algorithms_ComputeSubmergedVolume"), STAT_BuoyancyAlgorithms_ComputeSubmergedVolume, STATGROUP_Buoyancy);
+DECLARE_CYCLE_STAT(TEXT("Buoyancy_Algorithms_ComputeSubmergedVolume"), STAT_BuoyancyAlgorithms_ComputeSubmergedBounds, STATGROUP_Buoyancy);
+DECLARE_CYCLE_STAT(TEXT("Buoyancy_Algorithms_SubdivideBounds"), STAT_BuoyancyAlgorithms_SubdivideBounds, STATGROUP_Buoyancy);
+DECLARE_CYCLE_STAT(TEXT("Buoyancy_Algorithms_ComputeBuoyantForces"), STAT_BuoyancyAlgorithms_ComputeBuoyantForces, STATGROUP_Buoyancy);
+
diff --git a/Engine/Plugins/Experimental/Buoyancy/Source/Runtime/Public/BuoyancySubsystem.h b/Engine/Plugins/Experimental/Buoyancy/Source/Runtime/Public/BuoyancySubsystem.h
new file mode 100644
--- /dev/null
+++ b/Engine/Plugins/Experimental/Buoyancy/Source/Runtime/Public/BuoyancySubsystem.h
@@ -0,0 +1,112 @@
+// Copyright Epic Games, Inc. All Rights Reserved.
+
+#pragma once
+
+#include "Subsystems/WorldSubsystem.h"
+#include "Chaos/SimCallbackObject.h"
+#include "Chaos/Framework/PhysicsProxyBase.h"
+#include "PBDRigidsSolver.h"
+#include "Engine/EngineBaseTypes.h"
+#include "BuoyancySubsystem.generated.h"
+
+DECLARE_LOG_CATEGORY_EXTERN(LogBuoyancySubsystem, Log, All);
+
+//
+// Buoyancy Settings
+//
+
+struct FBuoyancySettings
+{
+	// Force buoyant particles which are in water to stay awake
+	bool bKeepAwake = false;
+
+	// Density of water is about 1g/cm^3
+	// Source: https://en.wikipedia.org/wiki/Properties_of_water
+	float WaterDensity = 0.0001f; // kg/cm^3
+
+	float MaxDeltaV = 200.f; // cm/s
+
+	float MaxDeltaW = 2.f; // rad/s
+
+	float WaterDrag = 1.f; // unitless
+
+	int32 MaxNumBoundsSubdivisions = 2;
+
+	float MinBoundsSubdivisionVol = FMath::Pow(100.f, 3.f); // 1m^3
+
+	ECollisionChannel WaterCollisionChannel = ECollisionChannel::ECC_MAX;
+};
+
+
+//
+// Buoyancy Subsystem
+//
+
+UCLASS()
+class BUOYANCY_API UBuoyancySubsystem : public UTickableWorldSubsystem
+{
+	GENERATED_BODY()
+
+		UBuoyancySubsystem()
+		: UTickableWorldSubsystem()
+		, bBuoyancySettingsChanged(false)
+		, BuoyancySettings(FBuoyancySettings())
+		, SimCallback(nullptr)
+	{ }
+
+protected:
+
+	// UTickableWorldSubsystem begin interface
+	virtual void PostInitialize() override;
+	virtual void Deinitialize() override;
+	virtual void Tick(float DeltaTime) override;
+	virtual TStatId GetStatId() const override;
+	// UTickableWorldSubsystem end interface
+
+
+private:
+
+	Chaos::FPhysicsSolver* GetSolver() const;
+
+	// When water plugin settings change, this callback will apply changes
+	void ApplyRuntimeSettings(const class UWaterRuntimeSettings* InSettings, EPropertyChangeType::Type ChangeType);
+
+	bool bBuoyancySettingsChanged;
+
+	FBuoyancySettings BuoyancySettings;
+
+	class FBuoyancySubsystemSimCallback* SimCallback;
+};
+
+
+//
+// Buoyancy Sim Callback
+//
+
+struct FBuoyancySubsystemSimCallbackInput : public Chaos::FSimCallbackInput
+{
+	// Here we use a unique ptr so that it is possible to provide an async
+	// input _without_ buoyancy settings (which may be eventually desirable
+	// when we eventually are passing lists of water bodies or water wave
+	// data).
+	mutable TUniquePtr<FBuoyancySettings> BuoyancySettings;
+
+	void Reset();
+};
+
+// NOTE: The Presimulate option is only needed for proper registry with the solver.
+//       We don't actually need (or want!) a presimulate tick.
+class FBuoyancySubsystemSimCallback : public Chaos::TSimCallbackObject<
+	FBuoyancySubsystemSimCallbackInput,
+	Chaos::FSimCallbackNoOutput,
+	Chaos::ESimCallbackOptions::Presimulate | Chaos::ESimCallbackOptions::MidPhaseModification>
+{
+private:
+	virtual void OnPreSimulate_Internal() override { }
+	virtual void OnMidPhaseModification_Internal(Chaos::FMidPhaseModifierAccessor& Modifier) override;
+
+	// Initially we won't have any settings - they have to get passed down
+	// via async input. I used TUniquePtr to control access to the same
+	// memory that was allocated by GT to minimize copies.
+	TUniquePtr<FBuoyancySettings> BuoyancySettings;
+};
diff --git a/Engine/Plugins/Experimental/Water/Source/Runtime/Public/WaterRuntimeSettings.h b/Engine/Plugins/Experimental/Water/Source/Runtime/Public/WaterRuntimeSettings.h
--- a/Engine/Plugins/Experimental/Water/Source/Runtime/Public/WaterRuntimeSettings.h
+++ b/Engine/Plugins/Experimental/Water/Source/Runtime/Public/WaterRuntimeSettings.h
@@ -66,6 +66,11 @@ public:
 	UPROPERTY(EditAnywhere, config, Category = Rendering)
 	float WaterBodyIconWorldZOffset = 75.0f;
 
+	/** Density of water used for buoyancy calculations. Given in gram per cubic centimeter */
+	// TODO: Add units to UnitConversions.h so that g/cm^3 can be explicit
+	UPROPERTY(EditAnywhere, config, Category = Buoyancy, meta=(ClampMin = 0, ForceUnits = "g/cm3"))
+	float WaterDensity = 1.f;
+
 #if WITH_EDITORONLY_DATA
 	// Delegate called whenever the curve data is updated
 	DECLARE_MULTICAST_DELEGATE_TwoParams(FOnUpdateSettings, const UWaterRuntimeSettings* /*Settings*/, EPropertyChangeType::Type /*ChangeType*/);
diff --git a/Engine/Source/Runtime/Engine/Public/Physics/PhysicsFiltering.h b/Engine/Source/Runtime/Engine/Public/Physics/PhysicsFiltering.h
--- a/Engine/Source/Runtime/Engine/Public/Physics/PhysicsFiltering.h
+++ b/Engine/Source/Runtime/Engine/Public/Physics/PhysicsFiltering.h
@@ -33,7 +33,7 @@ enum EPhysXFilterDataFlags
 // (ExtraFilter (top NumExtraFilterBits) + MyChannel (next NumCollisionChannelBits) as ECollisionChannel + Flags (remaining NumFilterDataFlagBits)
 // [NumExtraFilterBits] [NumCollisionChannelBits] [NumFilterDataFlagBits] = 32 bits
 enum { NumCollisionChannelBits = 5 };
-enum { NumFilterDataFlagBits = 32 - NumExtraFilterBits - NumCollisionChannelBits };
+enum { NumFilterDataFlagBits = 32 - ::NumExtraFilterBits - ::NumCollisionChannelBits };
 
 
 struct FPhysicsFilterBuilder
@@ -120,26 +120,26 @@ inline void CreateShapeFilterData(
 
 inline ECollisionChannel GetCollisionChannel(uint32 Word3)
 {
-	uint32 ChannelMask = (Word3 << NumExtraFilterBits) >> (32 - NumCollisionChannelBits);
+	uint32 ChannelMask = (Word3 << ::NumExtraFilterBits) >> (32 - ::NumCollisionChannelBits);
 	return (ECollisionChannel)ChannelMask;
 }
 
 inline ECollisionChannel GetCollisionChannelAndExtraFilter(uint32 Word3, FMaskFilter& OutMaskFilter)
 {
 	uint32 ChannelMask = GetCollisionChannel(Word3);
-	OutMaskFilter = Word3 >> (32 - NumExtraFilterBits);
+	OutMaskFilter = Word3 >> (32 - ::NumExtraFilterBits);
 	return (ECollisionChannel)ChannelMask;
 }
 
 inline uint32 CreateChannelAndFilter(ECollisionChannel CollisionChannel, FMaskFilter MaskFilter)
 {
-	uint32 ResultMask = (uint32(MaskFilter) << NumCollisionChannelBits) | (uint32)CollisionChannel;
+	uint32 ResultMask = (uint32(MaskFilter) << ::NumCollisionChannelBits) | (uint32)CollisionChannel;
 	return ResultMask << NumFilterDataFlagBits;
 }
 
 inline void UpdateMaskFilter(uint32& Word3, FMaskFilter NewMaskFilter)
 {
-	static_assert(NumExtraFilterBits <= 8, "Only up to 8 extra filter bits are supported.");
-	Word3 &= (0xFFFFFFFFu >> NumExtraFilterBits);	//we drop the top NumExtraFilterBits bits because that's where the new mask filter is going
-	Word3 |= uint32(NewMaskFilter) << (32 - NumExtraFilterBits);
+	static_assert(::NumExtraFilterBits <= 8, "Only up to 8 extra filter bits are supported.");
+	Word3 &= (0xFFFFFFFFu >> ::NumExtraFilterBits);	//we drop the top NumExtraFilterBits bits because that's where the new mask filter is going
+	Word3 |= uint32(NewMaskFilter) << (32 - ::NumExtraFilterBits);
 }
