Remove IsEnabled and SetEnabled from BuoyancySubsystem - they are not used anywhere

diff --git a/Engine/Plugins/Experimental/Buoyancy/Source/Runtime/Private/BuoyancySubsystem.cpp b/Engine/Plugins/Experimental/Buoyancy/Source/Runtime/Private/BuoyancySubsystem.cpp
--- a/Engine/Plugins/Experimental/Buoyancy/Source/Runtime/Private/BuoyancySubsystem.cpp
+++ b/Engine/Plugins/Experimental/Buoyancy/Source/Runtime/Private/BuoyancySubsystem.cpp
@@ -35,16 +35,6 @@ DEFINE_LOG_CATEGORY(LogBuoyancySubsystem);
 // Buoyancy Subsystem
 //
 
-void UBuoyancySubsystem::SetEnabled(bool bEnabled)
-{
-	
-}
-
-bool UBuoyancySubsystem::IsEnabled() const
-{
-	SimCallback != nullptr;
-}
-
 bool UBuoyancySubsystem::CreateSimCallback()
 {
 	// Create sim callback
diff --git a/Engine/Plugins/Experimental/Buoyancy/Source/Runtime/Public/BuoyancySubsystem.h b/Engine/Plugins/Experimental/Buoyancy/Source/Runtime/Public/BuoyancySubsystem.h
--- a/Engine/Plugins/Experimental/Buoyancy/Source/Runtime/Public/BuoyancySubsystem.h
+++ b/Engine/Plugins/Experimental/Buoyancy/Source/Runtime/Public/BuoyancySubsystem.h
@@ -54,11 +54,6 @@ class BUOYANCY_API UBuoyancySubsystem : public UTickableWorldSubsystem
 		, SimCallback(nullptr)
 	{ }
 
-public:
-
-	void SetEnabled(bool bEnabled);
-	bool IsEnabled() const;
-
 protected:
 
 	// UTickableWorldSubsystem begin interface
