Give buoyancy plugin its own set of runtime settings which includes an enabled state.

diff --git a/Engine/Plugins/Experimental/Buoyancy/Source/Runtime/Buoyancy.Build.cs b/Engine/Plugins/Experimental/Buoyancy/Source/Runtime/Buoyancy.Build.cs
--- a/Engine/Plugins/Experimental/Buoyancy/Source/Runtime/Buoyancy.Build.cs
+++ b/Engine/Plugins/Experimental/Buoyancy/Source/Runtime/Buoyancy.Build.cs
@@ -22,6 +22,7 @@ public class Buoyancy : ModuleRules
 				"ChaosCore",
 				"Chaos",
 				"PhysicsCore",
+				"DeveloperSettings",
 			}
 		);
 	}
diff --git a/Engine/Plugins/Experimental/Buoyancy/Source/Runtime/Private/BuoyancyRuntimeSettings.cpp b/Engine/Plugins/Experimental/Buoyancy/Source/Runtime/Private/BuoyancyRuntimeSettings.cpp
new file mode 100644
--- /dev/null
+++ b/Engine/Plugins/Experimental/Buoyancy/Source/Runtime/Private/BuoyancyRuntimeSettings.cpp
@@ -0,0 +1,26 @@
+// Copyright Epic Games, Inc. All Rights Reserved.
+
+#include "BuoyancyRuntimeSettings.h"
+
+FName UBuoyancyRuntimeSettings::GetCategoryName() const
+{
+	return FName(TEXT("Plugins"));
+}
+
+void UBuoyancyRuntimeSettings::PostInitProperties()
+{
+	Super::PostInitProperties();
+}
+
+#if WITH_EDITOR
+
+UBuoyancyRuntimeSettings::FOnUpdateSettings UBuoyancyRuntimeSettings::OnSettingsChange;
+
+void UBuoyancyRuntimeSettings::PostEditChangeProperty(struct FPropertyChangedEvent& PropertyChangedEvent)
+{
+	Super::PostEditChangeProperty(PropertyChangedEvent);
+
+	OnSettingsChange.Broadcast(this, PropertyChangedEvent.ChangeType);
+}
+
+#endif // WITH_EDITOR
diff --git a/Engine/Plugins/Experimental/Buoyancy/Source/Runtime/Private/BuoyancySubsystem.cpp b/Engine/Plugins/Experimental/Buoyancy/Source/Runtime/Private/BuoyancySubsystem.cpp
--- a/Engine/Plugins/Experimental/Buoyancy/Source/Runtime/Private/BuoyancySubsystem.cpp
+++ b/Engine/Plugins/Experimental/Buoyancy/Source/Runtime/Private/BuoyancySubsystem.cpp
@@ -3,7 +3,7 @@
 #include "BuoyancySubsystem.h"
 #include "BuoyancyAlgorithms.h"
 #include "BuoyancyStats.h"
-#include "WaterRuntimeSettings.h"
+#include "BuoyancyRuntimeSettings.h"
 #include "Engine/World.h"
 #include "Physics/Experimental/PhysScene_Chaos.h"
 #include "Physics/PhysicsFiltering.h"
@@ -20,9 +20,6 @@
 bool bBuoyancyDebugDraw = false;
 FAutoConsoleVariableRef CVarBuoyancyDebugDraw(TEXT("p.Buoyancy.DebugDraw"), bBuoyancyDebugDraw, TEXT(""));
 
-bool bBuoyancyEnabled = true;
-FAutoConsoleVariableRef CVarBuoyancyEnabled(TEXT("p.Buoyancy.Enabled"), bBuoyancyEnabled, TEXT(""));
-
 
 //
 // Logging
@@ -35,6 +32,31 @@ DEFINE_LOG_CATEGORY(LogBuoyancySubsystem);
 // Buoyancy Subsystem
 //
 
+bool UBuoyancySubsystem::SetEnabled(const bool bEnabled)
+{
+	if (bEnabled != IsEnabled())
+	{
+		if (bEnabled)
+		{
+			CreateSimCallback();
+		}
+		else
+		{
+			DestroySimCallback();
+		}
+
+		return IsEnabled() == bEnabled;
+	}
+
+	// Already had whatever setting
+	return true;
+}
+
+bool UBuoyancySubsystem::IsEnabled() const
+{
+	return SimCallback != nullptr;
+}
+
 bool UBuoyancySubsystem::CreateSimCallback()
 {
 	// Create sim callback
@@ -70,11 +92,11 @@ void UBuoyancySubsystem::PostInitialize()
 	CreateSimCallback();
 
 	// Apply initial runtime settings
-	ApplyRuntimeSettings(GetDefault<UWaterRuntimeSettings>(), EPropertyChangeType::ValueSet);
+	ApplyRuntimeSettings(GetDefault<UBuoyancyRuntimeSettings>(), EPropertyChangeType::ValueSet);
 
 	// Set up callback for when runtime settings change in editor
 #if WITH_EDITOR
-	GetDefault<UWaterRuntimeSettings>()->OnSettingsChange.AddUObject(this, &UBuoyancySubsystem::ApplyRuntimeSettings);
+	GetDefault<UBuoyancyRuntimeSettings>()->OnSettingsChange.AddUObject(this, &UBuoyancySubsystem::ApplyRuntimeSettings);
 #endif //WITH_EDITOR
 }
 
@@ -85,7 +107,7 @@ void UBuoyancySubsystem::Deinitialize()
 	Super::Deinitialize();
 }
 
-void UBuoyancySubsystem::ApplyRuntimeSettings(const UWaterRuntimeSettings* InSettings, EPropertyChangeType::Type ChangeType)
+void UBuoyancySubsystem::ApplyRuntimeSettings(const UBuoyancyRuntimeSettings* InSettings, EPropertyChangeType::Type ChangeType)
 {
 	bBuoyancySettingsChanged = true;
 
@@ -97,6 +119,9 @@ void UBuoyancySubsystem::ApplyRuntimeSettings(const UWaterRuntimeSettings* InSet
 		EUnit::KilogramsPerCubicCentimeter);
 
 	BuoyancySettings.WaterCollisionChannel = InSettings->CollisionChannelForWaterObjects;
+
+	// Enable or disable
+	SetEnabled(InSettings->bBuoyancyEnabled);
 }
 
 void UBuoyancySubsystem::Tick(float DeltaTime)
@@ -118,11 +143,6 @@ void UBuoyancySubsystem::Tick(float DeltaTime)
 
 	Super::Tick(DeltaTime);
 
-	if (bBuoyancyEnabled == false)
-	{
-		return;
-	}
-
 	if (SimCallback == nullptr)
 	{
 		return;
@@ -209,12 +229,6 @@ void FBuoyancySubsystemSimCallback::OnMidPhaseModification_Internal(Chaos::FMidP
 		return;
 	}
 
-	// If Buoyancy is disabled, skip
-	if (bBuoyancyEnabled == false)
-	{
-		return;
-	}
-
 	// Get the evolution
 	FPBDRigidsEvolution* Evolution = nullptr;
 	if (FPhysicsSolverBase* SolverBase = GetSolver())
diff --git a/Engine/Plugins/Experimental/Buoyancy/Source/Runtime/Public/BuoyancyRuntimeSettings.h b/Engine/Plugins/Experimental/Buoyancy/Source/Runtime/Public/BuoyancyRuntimeSettings.h
new file mode 100644
--- /dev/null
+++ b/Engine/Plugins/Experimental/Buoyancy/Source/Runtime/Public/BuoyancyRuntimeSettings.h
@@ -0,0 +1,45 @@
+// Copyright Epic Games, Inc. All Rights Reserved.
+
+#pragma once
+
+#include "Engine/DeveloperSettings.h"
+#include "Templates/SubclassOf.h"
+#include "BuoyancyRuntimeSettings.generated.h"
+
+enum ECollisionChannel : int;
+
+UCLASS(config = Engine, defaultconfig, meta = (DisplayName = "Buoyancy"))
+class BUOYANCY_API UBuoyancyRuntimeSettings : public UDeveloperSettings
+{
+	GENERATED_BODY()
+
+public:
+
+	virtual FName GetCategoryName() const;
+
+	// Whether or not the buoyancy subsystem should run.
+	// The subsystem has a public interface for enabling/disabling, so
+	// it's possible that a game overrides this at runtime
+	UPROPERTY(EditAnywhere, config, Category = Buoyancy, meta = (ClampMin = 0, ForceUnits = "g/cm3"))
+	bool bBuoyancyEnabled = true;
+
+	// Density of water to use in buoyancy calculations
+	UPROPERTY(EditAnywhere, config, Category = Buoyancy, meta = (ClampMin = 0, ForceUnits = "g/cm3"))
+	float WaterDensity = 1.f;
+
+	/** Collision channel to use for water ObjectTypes */
+	UPROPERTY(EditAnywhere, config, Category = Collision)
+	TEnumAsByte<ECollisionChannel> CollisionChannelForWaterObjects;
+
+	virtual void PostInitProperties() override;
+
+#if WITH_EDITOR
+	virtual void PostEditChangeProperty(struct FPropertyChangedEvent& PropertyChangedEvent) override;
+#endif // WITH_EDITOR
+
+#if WITH_EDITORONLY_DATA
+	// Delegate called whenever the curve data is updated
+	DECLARE_MULTICAST_DELEGATE_TwoParams(FOnUpdateSettings, const UBuoyancyRuntimeSettings* /*Settings*/, EPropertyChangeType::Type /*ChangeType*/);
+	static FOnUpdateSettings OnSettingsChange;
+#endif
+};
diff --git a/Engine/Plugins/Experimental/Buoyancy/Source/Runtime/Public/BuoyancySubsystem.h b/Engine/Plugins/Experimental/Buoyancy/Source/Runtime/Public/BuoyancySubsystem.h
--- a/Engine/Plugins/Experimental/Buoyancy/Source/Runtime/Public/BuoyancySubsystem.h
+++ b/Engine/Plugins/Experimental/Buoyancy/Source/Runtime/Public/BuoyancySubsystem.h
@@ -54,6 +54,15 @@ class BUOYANCY_API UBuoyancySubsystem : public UTickableWorldSubsystem
 		, SimCallback(nullptr)
 	{ }
 
+public:
+
+	// Return true if enable/disable was successful, or if
+	// we were already in the target state.
+	bool SetEnabled(const bool bEnabled);
+
+	// Return true if subsystem is enabled and running
+	bool IsEnabled() const;
+
 protected:
 
 	// UTickableWorldSubsystem begin interface
@@ -71,7 +80,7 @@ private:
 	Chaos::FPhysicsSolver* GetSolver() const;
 
 	// When water plugin settings change, this callback will apply changes
-	void ApplyRuntimeSettings(const class UWaterRuntimeSettings* InSettings, EPropertyChangeType::Type ChangeType);
+	void ApplyRuntimeSettings(const class UBuoyancyRuntimeSettings* InSettings, EPropertyChangeType::Type ChangeType);
 
 	bool bBuoyancySettingsChanged;
 
diff --git a/Engine/Plugins/Experimental/Water/Source/Runtime/Public/WaterRuntimeSettings.h b/Engine/Plugins/Experimental/Water/Source/Runtime/Public/WaterRuntimeSettings.h
--- a/Engine/Plugins/Experimental/Water/Source/Runtime/Public/WaterRuntimeSettings.h
+++ b/Engine/Plugins/Experimental/Water/Source/Runtime/Public/WaterRuntimeSettings.h
@@ -54,10 +54,6 @@ public:
 	UPROPERTY(EditAnywhere, config, Category = Collision)
 	TEnumAsByte<ECollisionChannel> CollisionChannelForWaterTraces;
 
-	/** Collision channel to use for water ObjectTypes */
-	UPROPERTY(EditAnywhere, config, Category = Collision)
-	TEnumAsByte<ECollisionChannel> CollisionChannelForWaterObjects;
-
 	/** Material Parameter Collection for everything water-related */
 	UPROPERTY(EditAnywhere, config, Category = Rendering)
 	TSoftObjectPtr<UMaterialParameterCollection> MaterialParameterCollection;
@@ -66,11 +62,6 @@ public:
 	UPROPERTY(EditAnywhere, config, Category = Rendering)
 	float WaterBodyIconWorldZOffset = 75.0f;
 
-	/** Density of water used for buoyancy calculations. Given in gram per cubic centimeter */
-	// TODO: Add units to UnitConversions.h so that g/cm^3 can be explicit
-	UPROPERTY(EditAnywhere, config, Category = Buoyancy, meta=(ClampMin = 0, ForceUnits = "g/cm3"))
-	float WaterDensity = 1.f;
-
 #if WITH_EDITORONLY_DATA
 	// Delegate called whenever the curve data is updated
 	DECLARE_MULTICAST_DELEGATE_TwoParams(FOnUpdateSettings, const UWaterRuntimeSettings* /*Settings*/, EPropertyChangeType::Type /*ChangeType*/);
