Expose a number of buoyancy options via runtime settings

diff --git a/Engine/Plugins/Experimental/Buoyancy/Source/Runtime/Private/BuoyancySubsystem.cpp b/Engine/Plugins/Experimental/Buoyancy/Source/Runtime/Private/BuoyancySubsystem.cpp
--- a/Engine/Plugins/Experimental/Buoyancy/Source/Runtime/Private/BuoyancySubsystem.cpp
+++ b/Engine/Plugins/Experimental/Buoyancy/Source/Runtime/Private/BuoyancySubsystem.cpp
@@ -117,8 +117,10 @@ void UBuoyancySubsystem::ApplyRuntimeSettings(const UBuoyancyRuntimeSettings* In
 	BuoyancySettings.WaterDensity = FUnitConversion::Convert(InSettings->WaterDensity,
 		EUnit::GramsPerCubicCentimeter,
 		EUnit::KilogramsPerCubicCentimeter);
-
 	BuoyancySettings.WaterCollisionChannel = InSettings->CollisionChannelForWaterObjects;
+	BuoyancySettings.bKeepAwake = InSettings->bKeepFloatingObjectsAwake;
+	BuoyancySettings.MaxNumBoundsSubdivisions = InSettings->MaxNumBoundsSubdivisions;
+	BuoyancySettings.MinBoundsSubdivisionVol = InSettings->MinBoundsSubdivisionVol;
 
 	// Enable or disable
 	SetEnabled(InSettings->bBuoyancyEnabled);
diff --git a/Engine/Plugins/Experimental/Buoyancy/Source/Runtime/Public/BuoyancyRuntimeSettings.h b/Engine/Plugins/Experimental/Buoyancy/Source/Runtime/Public/BuoyancyRuntimeSettings.h
--- a/Engine/Plugins/Experimental/Buoyancy/Source/Runtime/Public/BuoyancyRuntimeSettings.h
+++ b/Engine/Plugins/Experimental/Buoyancy/Source/Runtime/Public/BuoyancyRuntimeSettings.h
@@ -17,16 +17,29 @@ public:
 
 	virtual FName GetCategoryName() const;
 
-	// Whether or not the buoyancy subsystem should run.
-	// The subsystem has a public interface for enabling/disabling, so
-	// it's possible that a game overrides this at runtime
-	UPROPERTY(EditAnywhere, config, Category = Buoyancy, meta = (ClampMin = 0, ForceUnits = "g/cm3"))
+	/** Whether or not the buoyancy subsystem should run.
+	    The subsystem has a public interface for enabling/disabling, so
+	    it's possible that a game overrides this at runtime */
+	UPROPERTY(EditAnywhere, config, Category = Buoyancy, Meta = (ClampMin = 0, ForceUnits = "g/cm3"))
 	bool bBuoyancyEnabled = true;
 
-	// Density of water to use in buoyancy calculations
-	UPROPERTY(EditAnywhere, config, Category = Buoyancy, meta = (ClampMin = 0, ForceUnits = "g/cm3"))
+	/** Prevent floating particles from falling asleep by explicitly setting
+	    object state to dynamic when they are detected. */
+	UPROPERTY(EditAnywhere, config, Category = Buoyancy)
+	bool bKeepFloatingObjectsAwake = false;
+
+	/** Density of water to use in buoyancy calculations. The density of water is approximately 1g/cm3 */
+	UPROPERTY(EditAnywhere, config, Category = Buoyancy, Meta = (ClampMin = 0, ForceUnits = "g/cm3"))
 	float WaterDensity = 1.f;
 
+	/** Maximum number of times that a buoyancy bounds can be split into 8 */
+	UPROPERTY(EditAnywhere, config, Category = Buoyancy, Meta = (ClampMin = 0))
+	int32 MaxNumBoundsSubdivisions = 2;
+
+	/** Minimum volume bounding box which can be produced by a subdivision operation */
+	UPROPERTY(EditAnywhere, Config, Category = Buoyancy)
+	float MinBoundsSubdivisionVol = FMath::Pow(100.f, 3.f); // 1m^3
+
 	/** Collision channel to use for water ObjectTypes */
 	UPROPERTY(EditAnywhere, config, Category = Collision)
 	TEnumAsByte<ECollisionChannel> CollisionChannelForWaterObjects;
