Fix buoyancy volume double-counting.

This works by adding a sparsearray of bitarrays. The sparsearray is indexed by the particle unique indices. The bitarrays are indexed by shape indices. If a bit is set, then skip processing that shape. If not, then process a shape and set it.

This change also introduces a division between application of buoyant forces and dedection of submersions, in hopes that we can remove the manual integration bit at some point.

diff --git a/Engine/Plugins/Experimental/Buoyancy/Source/Runtime/Private/BuoyancyAlgorithms.cpp b/Engine/Plugins/Experimental/Buoyancy/Source/Runtime/Private/BuoyancyAlgorithms.cpp
--- a/Engine/Plugins/Experimental/Buoyancy/Source/Runtime/Private/BuoyancyAlgorithms.cpp
+++ b/Engine/Plugins/Experimental/Buoyancy/Source/Runtime/Private/BuoyancyAlgorithms.cpp
@@ -23,6 +23,34 @@ namespace
 {
 	using namespace Chaos;
 
+	// Check to see if an object's shape is marked as already submerged
+	bool IsShapeSubmerged_Internal(const TSparseArray<TBitArray<>>& SubmergedShapes, const int32 ParticleIndex, const int32 ShapeIndex)
+	{
+		return
+			SubmergedShapes.IsValidIndex(ParticleIndex) &&
+			SubmergedShapes[ParticleIndex].IsValidIndex(ShapeIndex) &&
+			SubmergedShapes[ParticleIndex][ShapeIndex];
+	}
+
+	// Mark an object's shape as submerged
+	void SubmergeShape_Internal(TSparseArray<TBitArray<>>& SubmergedShapes, const int32 ParticleIndex, const int32 ShapeIndex)
+	{
+		// If no shapes are tracked for this particle yet, add a bit array for it
+		if (SubmergedShapes.IsValidIndex(ParticleIndex) == false)
+		{
+			SubmergedShapes.Insert(ParticleIndex, TBitArray<>(false, ShapeIndex + 1));
+		}
+
+		// If the bit array for this particle existed already but is too small, expand it
+		else if (SubmergedShapes[ParticleIndex].IsValidIndex(ShapeIndex) == false)
+		{
+			SubmergedShapes[ParticleIndex].SetNum(ShapeIndex + 1, false);
+		}
+
+		// Mark this particle's shape as submerged
+		SubmergedShapes[ParticleIndex][ShapeIndex] = true;
+	}
+
 	// NOTE: See SubdivideBounds(...) for a description of this algorithm
 	//
 	// TODO: Use an array + offset rather than raw ptr
@@ -79,13 +107,15 @@ namespace BuoyancyAlgorithms
 {
 	using namespace Chaos;
 
-	bool ComputeSubmergedVolume(const FGeometryParticleHandle* ParticleA, const FGeometryParticleHandle* ParticleB, const int32 NumSubdivisions, const float MinVolume, float& SubmergedVol, FVec3& SubmergedCoM)
+	bool ComputeSubmergedVolume(const FGeometryParticleHandle* ParticleA, const FGeometryParticleHandle* ParticleB, const int32 NumSubdivisions, const float MinVolume, TSparseArray<TBitArray<>>& SubmergedShapes, float& SubmergedVol, FVec3& SubmergedCoM)
 	{
 		SCOPE_CYCLE_COUNTER(STAT_BuoyancyAlgorithms_ComputeSubmergedVolume)
 
 		SubmergedVol = 0.f;
 		SubmergedCoM = FVec3::ZeroVector;
 
+		const int32 ParticleIndexB = ParticleB->UniqueIdx().Idx;
+
 		const FImplicitObject* RootImplicitA = ParticleA->GetGeometry();
 		const FImplicitObject* RootImplicitB = ParticleB->GetGeometry();
 
@@ -105,8 +135,8 @@ namespace BuoyancyAlgorithms
 		// object hierarchy of ParticleB.
 		RootImplicitA->VisitLeafObjects(
 			[ParticleA, &ShapeInstancesA, &ParticleWorldTransformA,
-			ParticleB, &ShapeInstancesB, RootImplicitB, &ParticleWorldTransformB, &ParticleTransformAToB,
-			&NumSubdivisions, &MinVolume, &SubmergedVol, &SubmergedCoM]
+			ParticleIndexB, &ShapeInstancesB, RootImplicitB, &ParticleWorldTransformB, &ParticleTransformAToB,
+			&NumSubdivisions, &MinVolume, &SubmergedShapes, &SubmergedVol, &SubmergedCoM]
 			(const FImplicitObject* ImplicitA, const FRigidTransform3& RelativeTransformA, const int32 RootObjectIndexA, const int32 ObjectIndex, const int32 LeafObjectIndexA)
 		{
 			const FAABB3 RelativeBoundsA = ImplicitA->CalculateTransformedBounds(RelativeTransformA);
@@ -121,11 +151,19 @@ namespace BuoyancyAlgorithms
 
 			// Detect collisions between ImplicitA and Implicit Hierarchy of ParticleB
 			RootImplicitB->VisitOverlappingLeafObjects(ShapeBoundsAInB,
-				[ParticleA, ImplicitA, ShapeInstanceA, &ParticleWorldTransformA, &RelativeTransformA, LeafObjectIndexA,
-				ParticleB, &ParticleWorldTransformB, &ShapeInstancesB,
-				&NumSubdivisions, &MinVolume, &SubmergedVol, &SubmergedCoM]
+				[ParticleA, ImplicitA, ShapeInstanceA, LeafObjectIndexA,
+				&ParticleWorldTransformA, &RelativeTransformA,
+				ParticleIndexB, &ParticleWorldTransformB,
+				&NumSubdivisions, &MinVolume, &SubmergedShapes, &SubmergedVol, &SubmergedCoM]
 				(const FImplicitObject* ImplicitB, const FRigidTransform3& RelativeTransformB, const int32 RootObjectIndexB, const int32 ObjectIndexB, const int32 LeafObjectIndexB)
 			{
+				// If this shape has already been submerged, skip it to avoid double-counting
+				// any buoyancy contributions.
+				if (IsShapeSubmerged_Internal(SubmergedShapes, ParticleIndexB, ObjectIndexB))
+				{
+					return;
+				}
+
 				// Get shape world transforms
 				const FRigidTransform3 ShapeWorldTransformA = RelativeTransformA * ParticleWorldTransformA;
 				const FRigidTransform3 ShapeWorldTransformB = RelativeTransformB * ParticleWorldTransformB;
@@ -151,6 +189,7 @@ namespace BuoyancyAlgorithms
 				SubdivideBounds(BoxBInB, NumSubdivisions, MinVolume, BoxesInB);
 
 				// Loop over every subdivision of the shape bounds, counting up submerged portions
+				bool bSubmerged = false;
 				for (const FAABB3& BoxInB : BoxesInB)
 				{
 #if ENABLE_DRAW_DEBUG
@@ -170,6 +209,9 @@ namespace BuoyancyAlgorithms
 					FAABB3 SubmergedBoundsInB;
 					if (ComputeSubmergedBounds(WaterZ, BoxInB, ShapeWorldTransformB, SubmergedBoundsInB))
 					{
+						// At this point we know that the shape is submerged
+						bSubmerged = true;
+
 						// This bounds box is submerged. Compute it's volume and center of mass
 						// in world space, and add those contributions to the submerged quantity.
 						const FVec3 LeafSubmergedCoM = ShapeWorldTransformB.TransformPosition(SubmergedBoundsInB.GetCenter());
@@ -196,24 +238,30 @@ namespace BuoyancyAlgorithms
 					}
 				}
 
-#if ENABLE_DRAW_DEBUG
-				if (bBuoyancyDebugDraw)
+				// If the shape pair overlapped, flag the submersion
+				if (bSubmerged)
 				{
-					// Water OOBB
-					Chaos::FDebugDrawQueue::GetInstance().DrawDebugBox(
-						ShapeWorldTransformA.TransformPosition(BoxAInA.GetCenter()),
-						BoxAInA.Extents() * .5f,
-						ShapeWorldTransformA.GetRotation(),
-						FColor::Cyan, false, -1.f, SDPG_Foreground, 2.f);
-
-					// Rigid OOBB
-					Chaos::FDebugDrawQueue::GetInstance().DrawDebugBox(
-						ShapeWorldTransformB.TransformPosition(BoxBInB.GetCenter()),
-						BoxBInB.Extents() * .5f,
-						ShapeWorldTransformB.GetRotation(),
-						FColor::Green, false, -1.f, SDPG_Foreground, 2.f);
-				}
+					SubmergeShape_Internal(SubmergedShapes, ParticleIndexB, ObjectIndexB);
+
+#if ENABLE_DRAW_DEBUG
+					if (bBuoyancyDebugDraw)
+					{
+						// Water OOBB
+						Chaos::FDebugDrawQueue::GetInstance().DrawDebugBox(
+							ShapeWorldTransformA.TransformPosition(BoxAInA.GetCenter()),
+							BoxAInA.Extents() * .5f,
+							ShapeWorldTransformA.GetRotation(),
+							FColor::Cyan, false, -1.f, SDPG_Foreground, 2.f);
+
+						// Rigid OOBB
+						Chaos::FDebugDrawQueue::GetInstance().DrawDebugBox(
+							ShapeWorldTransformB.TransformPosition(BoxBInB.GetCenter()),
+							BoxBInB.Extents() * .5f,
+							ShapeWorldTransformB.GetRotation(),
+							FColor::Green, false, -1.f, SDPG_Foreground, 2.f);
+					}
 #endif
+				}
 			});
 		});
 
diff --git a/Engine/Plugins/Experimental/Buoyancy/Source/Runtime/Private/BuoyancySubsystem.cpp b/Engine/Plugins/Experimental/Buoyancy/Source/Runtime/Private/BuoyancySubsystem.cpp
--- a/Engine/Plugins/Experimental/Buoyancy/Source/Runtime/Private/BuoyancySubsystem.cpp
+++ b/Engine/Plugins/Experimental/Buoyancy/Source/Runtime/Private/BuoyancySubsystem.cpp
@@ -210,11 +210,11 @@ namespace
 	};
 }
 
-void FBuoyancySubsystemSimCallback::OnMidPhaseModification_Internal(Chaos::FMidPhaseModifierAccessor& MidPhaseAccessor)
+void FBuoyancySubsystemSimCallback::OnPreSimulate_Internal()
 {
 	using namespace Chaos;
 
-	SCOPE_CYCLE_COUNTER(STAT_BuoyancySubsystem_OnMidPhaseModification)
+	SCOPE_CYCLE_COUNTER(STAT_BuoyancySubsystem_OnPreSimulate)
 
 	// If we were sent new buoyancy settings or data, update our local sim copy
 	if (const FBuoyancySubsystemSimCallbackInput* Input = GetConsumerInput_Internal())
@@ -255,17 +255,67 @@ void FBuoyancySubsystemSimCallback::OnMidPhaseModification_Internal(Chaos::FMidP
 	// How much time has the sim ticked this frame
 	const FReal DeltaSeconds = GetDeltaTime_Internal();
 
-	// TODO:
+	// Apply all buoyant forces
+	for (const FSubmersion& Submersion : Submersions)
+	{
+		// Figure out the gravity level of the particle
+		const int32 GravityGroupIndex = Submersion.Particle->GravityGroupIndex();
+		const FVec3 GravityAccel
+			= PerParticleGravity != nullptr && GravityGroupIndex != INDEX_NONE
+			? (FVec3)PerParticleGravity->GetAcceleration(GravityGroupIndex)
+			: FVec3::DownVector * 980.f; // Default to "regular" gravity
+
+		// Compute delta linear and angular velocities due to buoyancy. If they're big enough to
+		// matter, apply them
+		FVec3 DeltaV, DeltaW;
+		if (BuoyancyAlgorithms::ComputeBuoyantForce(Submersion.Particle, DeltaSeconds, BuoyancySettings->WaterDensity, BuoyancySettings->WaterDrag, GravityAccel, Submersion.CoM, Submersion.Vol, DeltaV, DeltaW))
+		{
+			// Clamp delta velocities
+			DeltaV = DeltaV.GetClampedToSize(0.f, BuoyancySettings->MaxDeltaV);
+			DeltaW = DeltaW.GetClampedToSize(0.f, BuoyancySettings->MaxDeltaW);
+
+			// Apply the deltas
+			Submersion.Particle->SetV(Submersion.Particle->V() + DeltaV);
+			Submersion.Particle->SetW(Submersion.Particle->W() + DeltaW);
+
+			// Wake up the body??
+			if (Evolution && BuoyancySettings->bKeepAwake)
+			{
+				Evolution->SetParticleObjectState(Submersion.Particle, EObjectStateType::Dynamic);
+			}
+		}
+
+	}
+
+	// Clear submersions, but keep memory allocated
+	Submersions.Reset();
+}
+
+void FBuoyancySubsystemSimCallback::OnMidPhaseModification_Internal(Chaos::FMidPhaseModifierAccessor& MidPhaseAccessor)
+{
+	using namespace Chaos;
+
+	SCOPE_CYCLE_COUNTER(STAT_BuoyancySubsystem_OnMidPhaseModification)
+
+	// If we don't have a valid buoyancy settings object, don't continue
+	if (BuoyancySettings.IsValid() == false)
+	{
+		return;
+	}
+
+	// This 2d bit array is used to avoid double-counting buoyancy for any particular
+	// shape on a body. The outer TSparseArray is indexed by the unique idx on a particle,
+	// the inner bitarray is indexed by it's shape indices.
 	//
-	// To avoid double-counting any particular object for buoyancy, make a bit array
-	// indexed by particle indices. Every time we apply a buoyant force, we'll skip
-	// that particle next time we see it
+	// If we have computed a buoyant force for a particular shape on an object already
+	// then skip it on the second time around.
+	TSparseArray< TBitArray<> > SubmergedShapes;
 
 	// NOTE: For now we visit _every_ midphase, and check for ones which involve
 	// our target collision channel. It would be nice if it were possible instead
 	// to get a list of water body particles and loop over only midphases which
 	// involve them.
-	MidPhaseAccessor.VisitMidPhases([this, Evolution, PerParticleGravity, DeltaSeconds](Chaos::FMidPhaseModifier& MidPhase)
+	MidPhaseAccessor.VisitMidPhases([this, &SubmergedShapes](Chaos::FMidPhaseModifier& MidPhase)
 	{
 		SCOPE_CYCLE_COUNTER(STAT_BuoyancySubsystem_VisitMidphases)
 
@@ -315,7 +365,7 @@ void FBuoyancySubsystemSimCallback::OnMidPhaseModification_Internal(Chaos::FMidP
 		// Compute submerged volume and CoM
 		float SubmergedVol;
 		FVec3 SubmergedCoM;
-		if (BuoyancyAlgorithms::ComputeSubmergedVolume(WaterParticle, RigidParticle, BuoyancySettings->MaxNumBoundsSubdivisions, BuoyancySettings->MinBoundsSubdivisionVol, SubmergedVol, SubmergedCoM))
+		if (BuoyancyAlgorithms::ComputeSubmergedVolume(WaterParticle, RigidParticle, BuoyancySettings->MaxNumBoundsSubdivisions, BuoyancySettings->MinBoundsSubdivisionVol, SubmergedShapes, SubmergedVol, SubmergedCoM))
 		{
 			SCOPE_CYCLE_COUNTER(STAT_BuoyancySubsystem_ComputeBuoyantForces)
 
@@ -324,31 +374,27 @@ void FBuoyancySubsystemSimCallback::OnMidPhaseModification_Internal(Chaos::FMidP
 			// approximate submerged center of mass, and apply the buoyancy force there.
 			if (SubmergedVol > SMALL_NUMBER)
 			{
-				// Figure out the gravity level of the particle
-				const int32 GravityGroupIndex = RigidParticle->GravityGroupIndex();
-				const FVec3 GravityAccel
-					= PerParticleGravity != nullptr && GravityGroupIndex != INDEX_NONE
-					? (FVec3)PerParticleGravity->GetAcceleration(GravityGroupIndex)
-					: FVec3::DownVector * 980.f; // Default to "regular" gravity
-
-				// Compute delta linear and angular velocities due to buoyancy. If they're big enough to
-				// matter, apply them
-				FVec3 DeltaV, DeltaW;
-				if (BuoyancyAlgorithms::ComputeBuoyantForce(RigidParticle, DeltaSeconds, BuoyancySettings->WaterDensity, BuoyancySettings->WaterDrag, GravityAccel, SubmergedCoM, SubmergedVol, DeltaV, DeltaW))
+				const int32 RigidParticleIndex = RigidParticle->UniqueIdx().Idx;
+
+				// If this particle was already marked submerged, add to its existing submersion
+				if (Submersions.IsValidIndex(RigidParticleIndex))
+				{
+					FSubmersion& Submersion = Submersions[RigidParticleIndex];
+					ensureMsgf(Submersion.Particle == RigidParticle, TEXT("Something went wrong - there's a particle index mismatch in the Submersions sparse array"));
+
+					// Sum the volumes
+					Submersion.Vol = Submersion.Vol + SubmergedVol;
+
+					// Get the weighted-average CoM
+					// NOTE: The unchecked division should be safe since we already
+					// know SubmergedVol > SMALL_NUMBER
+					Submersion.CoM = ((Submersion.CoM * Submersion.Vol) + (SubmergedCoM * SubmergedVol)) / Submersion.Vol;
+				}
+
+				// If this particle was not yet submerged, make a new submersion for it
+				else
 				{
-					// Clamp delta velocities
-					DeltaV = DeltaV.GetClampedToSize(0.f, BuoyancySettings->MaxDeltaV);
-					DeltaW = DeltaW.GetClampedToSize(0.f, BuoyancySettings->MaxDeltaW);
-
-					// Apply the deltas
-					RigidParticle->SetV(RigidParticle->V() + DeltaV);
-					RigidParticle->SetW(RigidParticle->W() + DeltaW);
-
-					// Wake up the body??
-					if (Evolution && BuoyancySettings->bKeepAwake)
-					{
-						Evolution->SetParticleObjectState(RigidParticle, EObjectStateType::Dynamic);
-					}
+					Submersions.Insert(RigidParticleIndex, { RigidParticle, SubmergedVol, SubmergedCoM });
 				}
 			}
 		}
diff --git a/Engine/Plugins/Experimental/Buoyancy/Source/Runtime/Public/BuoyancyAlgorithms.h b/Engine/Plugins/Experimental/Buoyancy/Source/Runtime/Public/BuoyancyAlgorithms.h
--- a/Engine/Plugins/Experimental/Buoyancy/Source/Runtime/Public/BuoyancyAlgorithms.h
+++ b/Engine/Plugins/Experimental/Buoyancy/Source/Runtime/Public/BuoyancyAlgorithms.h
@@ -11,8 +11,21 @@ namespace BuoyancyAlgorithms
 {
 	using namespace Chaos;
 
+	// Minimal struct containing essential data about a particular submersion
+	struct FSubmersion
+	{
+		// Indicates the submerged particle
+		FPBDRigidParticleHandle* SubmergedParticle;
+
+		// Total submerged volume
+		float SubmergedVolume;
+
+		// Effective submerged center of mass
+		FVec3 SubmergedCoM;
+	};
+
 	// Compute an approximate volume and center of mass of particle B submerged in particle A
-	bool ComputeSubmergedVolume(const FGeometryParticleHandle* ParticleA, const FGeometryParticleHandle* ParticleB, int32 NumSubdivisions, float MinVolume, float& SubmergedVol, FVec3& SubmergedCoM);
+	bool ComputeSubmergedVolume(const FGeometryParticleHandle* ParticleA, const FGeometryParticleHandle* ParticleB, int32 NumSubdivisions, float MinVolume, TSparseArray<TBitArray<>>& SubmergedShapes, float& SubmergedVol, FVec3& SubmergedCoM);
 
 	// Given an OOBB and a water level, generate another OOBB which is 1. entirely contained
 	// within the input OOBB and 2. entirely contains the portion of the OOBB which is submerged
diff --git a/Engine/Plugins/Experimental/Buoyancy/Source/Runtime/Public/BuoyancyStats.h b/Engine/Plugins/Experimental/Buoyancy/Source/Runtime/Public/BuoyancyStats.h
--- a/Engine/Plugins/Experimental/Buoyancy/Source/Runtime/Public/BuoyancyStats.h
+++ b/Engine/Plugins/Experimental/Buoyancy/Source/Runtime/Public/BuoyancyStats.h
@@ -7,6 +7,7 @@
 DECLARE_STATS_GROUP(TEXT("Buoyancy"), STATGROUP_Buoyancy, STATCAT_Advanced);
 
 // Subsystem
+DECLARE_CYCLE_STAT(TEXT("Buoyancy_Subsystem_OnPreSimulate"), STAT_BuoyancySubsystem_OnPreSimulate, STATGROUP_Buoyancy);
 DECLARE_CYCLE_STAT(TEXT("Buoyancy_Subsystem_OnMidPhaseModification"), STAT_BuoyancySubsystem_OnMidPhaseModification, STATGROUP_Buoyancy);
 DECLARE_CYCLE_STAT(TEXT("Buoyancy_Subsystem_VisitMidphases"), STAT_BuoyancySubsystem_VisitMidphases, STATGROUP_Buoyancy);
 DECLARE_CYCLE_STAT(TEXT("Buoyancy_Subsystem_ComputeBuoyantForces"), STAT_BuoyancySubsystem_ComputeBuoyantForces, STATGROUP_Buoyancy);
diff --git a/Engine/Plugins/Experimental/Buoyancy/Source/Runtime/Public/BuoyancySubsystem.h b/Engine/Plugins/Experimental/Buoyancy/Source/Runtime/Public/BuoyancySubsystem.h
--- a/Engine/Plugins/Experimental/Buoyancy/Source/Runtime/Public/BuoyancySubsystem.h
+++ b/Engine/Plugins/Experimental/Buoyancy/Source/Runtime/Public/BuoyancySubsystem.h
@@ -113,11 +113,24 @@ class FBuoyancySubsystemSimCallback : public Chaos::TSimCallbackObject<
 	Chaos::ESimCallbackOptions::Presimulate | Chaos::ESimCallbackOptions::MidPhaseModification>
 {
 private:
-	virtual void OnPreSimulate_Internal() override { }
+
+	virtual void OnPreSimulate_Internal() override;
 	virtual void OnMidPhaseModification_Internal(Chaos::FMidPhaseModifierAccessor& Modifier) override;
 
 	// Initially we won't have any settings - they have to get passed down
 	// via async input. I used TUniquePtr to control access to the same
 	// memory that was allocated by GT to minimize copies.
 	TUniquePtr<FBuoyancySettings> BuoyancySettings;
+
+	// A minimal struct of data tracking all the submersions in a frame.
+	struct FSubmersion
+	{
+		Chaos::FPBDRigidParticleHandle* Particle;
+		float Vol;
+		Chaos::FVec3 CoM;
+	};
+
+	// This sparse array of submersion events is indexed on particle unique indices.
+	// All buoyant forces due to submersions are applied at once.
+	TSparseArray<FSubmersion> Submersions;
 };
