Fix buoyancy computing submerged volume for shapes not marked for physics collision, which resulted in another source of double-counting

diff --git a/Engine/Plugins/Experimental/Buoyancy/Source/Runtime/Private/BuoyancyAlgorithms.cpp b/Engine/Plugins/Experimental/Buoyancy/Source/Runtime/Private/BuoyancyAlgorithms.cpp
--- a/Engine/Plugins/Experimental/Buoyancy/Source/Runtime/Private/BuoyancyAlgorithms.cpp
+++ b/Engine/Plugins/Experimental/Buoyancy/Source/Runtime/Private/BuoyancyAlgorithms.cpp
@@ -157,6 +157,15 @@ namespace BuoyancyAlgorithms
 				&NumSubdivisions, &MinVolume, &SubmergedShapes, &SubmergedVol, &SubmergedCoM]
 				(const FImplicitObject* ImplicitB, const FRigidTransform3& RelativeTransformB, const int32 RootObjectIndexB, const int32 ObjectIndexB, const int32 LeafObjectIndexB)
 			{
+				// If this shape doesn't have collision then skip it
+				//
+				// NOTE: Do we possibly need to do this in a more sophisticated way, like with actual
+				// collision filtering?
+				if (ImplicitB->GetDoCollide() == false)
+				{
+					return;
+				}
+
 				// If this shape has already been submerged, skip it to avoid double-counting
 				// any buoyancy contributions.
 				if (IsShapeSubmerged_Internal(SubmergedShapes, ParticleIndexB, ObjectIndexB))
