Fix particle cache working backwards from MidPhase to PreSimulate. Not valid to work in that direction because PreSimulate happens at the beginning of the frame and MidPhase in the middle - particles may be reset before the next PreSimluate, so particles cached in MidPhase may no longer be valid

diff --git a/Engine/Plugins/Experimental/Buoyancy/Source/Runtime/Private/BuoyancySubsystem.cpp b/Engine/Plugins/Experimental/Buoyancy/Source/Runtime/Private/BuoyancySubsystem.cpp
--- a/Engine/Plugins/Experimental/Buoyancy/Source/Runtime/Private/BuoyancySubsystem.cpp
+++ b/Engine/Plugins/Experimental/Buoyancy/Source/Runtime/Private/BuoyancySubsystem.cpp
@@ -231,6 +231,22 @@
 		return;
 	}
 
+
+
+}
+
+void FBuoyancySubsystemSimCallback::OnMidPhaseModification_Internal(Chaos::FMidPhaseModifierAccessor& MidPhaseAccessor)
+{
+	using namespace Chaos;
+
+	SCOPE_CYCLE_COUNTER(STAT_BuoyancySubsystem_OnMidPhaseModification)
+
+	// If we don't have a valid buoyancy settings object, don't continue
+	if (BuoyancySettings.IsValid() == false)
+	{
+		return;
+	}
+
 	// Get the evolution
 	FPBDRigidsEvolution* Evolution = nullptr;
 	if (FPhysicsSolverBase* SolverBase = GetSolver())
@@ -255,67 +271,15 @@
 	// How much time has the sim ticked this frame
 	const FReal DeltaSeconds = GetDeltaTime_Internal();
 
-	// Apply all buoyant forces
-	for (const FSubmersion& Submersion : Submersions)
-	{
-		// Figure out the gravity level of the particle
-		const int32 GravityGroupIndex = Submersion.Particle->GravityGroupIndex();
-		const FVec3 GravityAccel
-			= PerParticleGravity != nullptr && GravityGroupIndex != INDEX_NONE
-			? (FVec3)PerParticleGravity->GetAcceleration(GravityGroupIndex)
-			: FVec3::DownVector * 980.f; // Default to "regular" gravity
-
-		// Compute delta linear and angular velocities due to buoyancy. If they're big enough to
-		// matter, apply them
-		FVec3 DeltaV, DeltaW;
-		if (BuoyancyAlgorithms::ComputeBuoyantForce(Submersion.Particle, DeltaSeconds, BuoyancySettings->WaterDensity, BuoyancySettings->WaterDrag, GravityAccel, Submersion.CoM, Submersion.Vol, DeltaV, DeltaW))
-		{
-			// Clamp delta velocities
-			DeltaV = DeltaV.GetClampedToSize(0.f, BuoyancySettings->MaxDeltaV);
-			DeltaW = DeltaW.GetClampedToSize(0.f, BuoyancySettings->MaxDeltaW);
-
-			// Apply the deltas
-			Submersion.Particle->SetV(Submersion.Particle->V() + DeltaV);
-			Submersion.Particle->SetW(Submersion.Particle->W() + DeltaW);
-
-			// Wake up the body??
-			if (Evolution && BuoyancySettings->bKeepAwake)
-			{
-				Evolution->SetParticleObjectState(Submersion.Particle, EObjectStateType::Dynamic);
-			}
-		}
-
-	}
-
-	// Clear submersions, but keep memory allocated
+	// Clear submersions and submerged shapes array, but keep their memory allocated
 	Submersions.Reset();
-}
+	SubmergedShapes.Reset();
 
-void FBuoyancySubsystemSimCallback::OnMidPhaseModification_Internal(Chaos::FMidPhaseModifierAccessor& MidPhaseAccessor)
-{
-	using namespace Chaos;
-
-	SCOPE_CYCLE_COUNTER(STAT_BuoyancySubsystem_OnMidPhaseModification)
-
-	// If we don't have a valid buoyancy settings object, don't continue
-	if (BuoyancySettings.IsValid() == false)
-	{
-		return;
-	}
-
-	// This 2d bit array is used to avoid double-counting buoyancy for any particular
-	// shape on a body. The outer TSparseArray is indexed by the unique idx on a particle,
-	// the inner bitarray is indexed by it's shape indices.
-	//
-	// If we have computed a buoyant force for a particular shape on an object already
-	// then skip it on the second time around.
-	TSparseArray< TBitArray<> > SubmergedShapes;
-
 	// NOTE: For now we visit _every_ midphase, and check for ones which involve
 	// our target collision channel. It would be nice if it were possible instead
 	// to get a list of water body particles and loop over only midphases which
 	// involve them.
-	MidPhaseAccessor.VisitMidPhases([this, &SubmergedShapes](Chaos::FMidPhaseModifier& MidPhase)
+	MidPhaseAccessor.VisitMidPhases([this](Chaos::FMidPhaseModifier& MidPhase)
 	{
 		SCOPE_CYCLE_COUNTER(STAT_BuoyancySubsystem_VisitMidphases)
 
@@ -399,4 +363,36 @@
 			}
 		}
 	});
+
+	// Apply all buoyant forces
+	for (const FSubmersion& Submersion : Submersions)
+	{
+		// Figure out the gravity level of the particle
+		const int32 GravityGroupIndex = Submersion.Particle->GravityGroupIndex();
+		const FVec3 GravityAccel
+			= PerParticleGravity != nullptr && GravityGroupIndex != INDEX_NONE
+			? (FVec3)PerParticleGravity->GetAcceleration(GravityGroupIndex)
+			: FVec3::DownVector * 980.f; // Default to "regular" gravity
+
+		// Compute delta linear and angular velocities due to buoyancy. If they're big enough to
+		// matter, apply them
+		FVec3 DeltaV, DeltaW;
+		if (BuoyancyAlgorithms::ComputeBuoyantForce(Submersion.Particle, DeltaSeconds, BuoyancySettings->WaterDensity, BuoyancySettings->WaterDrag, GravityAccel, Submersion.CoM, Submersion.Vol, DeltaV, DeltaW))
+		{
+			// Clamp delta velocities
+			DeltaV = DeltaV.GetClampedToSize(0.f, BuoyancySettings->MaxDeltaV);
+			DeltaW = DeltaW.GetClampedToSize(0.f, BuoyancySettings->MaxDeltaW);
+
+			// Apply the deltas
+			Submersion.Particle->SetV(Submersion.Particle->V() + DeltaV);
+			Submersion.Particle->SetW(Submersion.Particle->W() + DeltaW);
+
+			// Wake up the body??
+			if (Evolution && BuoyancySettings->bKeepAwake)
+			{
+				Evolution->SetParticleObjectState(Submersion.Particle, EObjectStateType::Dynamic);
+			}
+		}
+
+	}
 }
diff --git a/Engine/Plugins/Experimental/Buoyancy/Source/Runtime/Public/BuoyancySubsystem.h b/Engine/Plugins/Experimental/Buoyancy/Source/Runtime/Public/BuoyancySubsystem.h
--- a/Engine/Plugins/Experimental/Buoyancy/Source/Runtime/Public/BuoyancySubsystem.h
+++ b/Engine/Plugins/Experimental/Buoyancy/Source/Runtime/Public/BuoyancySubsystem.h
@@ -131,6 +131,19 @@ private:
 	};
 
 	// This sparse array of submersion events is indexed on particle unique indices.
-	// All buoyant forces due to submersions are applied at once.
+	// All buoyant forces due to submersions are applied at once. It's stored as
+	// a member variable and reset every frame, to avoid reallocation of similarly
+	// sized data.
 	TSparseArray<FSubmersion> Submersions;
+
+	// This is a sparse array of bit arrays representing which shapes in an object
+	// have already been accounted for when submerging an object. For example, if
+	// a massive BVH object has two leaf node shapes submerged in different pools
+	// of water and we've already detected that leaf A is submerged, we don't
+	// need to test A again. This helps to avoid double counting submerged shapes.
+	//
+	// Just like Submersions, we have this as a member variable only to keep the
+	// memory hot - the array is reset, repopulated and traversed, every frame,
+	// so we want to minimize allocations.
+	TSparseArray<TBitArray<>> SubmergedShapes;
 };
