Don't create buoyancy sim callback when buoyancy is disabled.

diff --git a/Engine/Plugins/Experimental/Buoyancy/Source/Runtime/Private/BuoyancySubsystem.cpp b/Engine/Plugins/Experimental/Buoyancy/Source/Runtime/Private/BuoyancySubsystem.cpp
--- a/Engine/Plugins/Experimental/Buoyancy/Source/Runtime/Private/BuoyancySubsystem.cpp
+++ b/Engine/Plugins/Experimental/Buoyancy/Source/Runtime/Private/BuoyancySubsystem.cpp
@@ -89,8 +89,6 @@ void UBuoyancySubsystem::PostInitialize()
 {
 	Super::PostInitialize();
 
-	CreateSimCallback();
-
 	// Apply initial runtime settings
 	ApplyRuntimeSettings(GetDefault<UBuoyancyRuntimeSettings>(), EPropertyChangeType::ValueSet);
 
diff --git a/Engine/Plugins/Experimental/Buoyancy/Source/Runtime/Public/BuoyancyRuntimeSettings.h b/Engine/Plugins/Experimental/Buoyancy/Source/Runtime/Public/BuoyancyRuntimeSettings.h
--- a/Engine/Plugins/Experimental/Buoyancy/Source/Runtime/Public/BuoyancyRuntimeSettings.h
+++ b/Engine/Plugins/Experimental/Buoyancy/Source/Runtime/Public/BuoyancyRuntimeSettings.h
@@ -8,7 +8,7 @@
 
 enum ECollisionChannel : int;
 
-UCLASS(config = Engine, defaultconfig, meta = (DisplayName = "Buoyancy"))
+UCLASS(Config = Engine, DefaultConfig, Meta = (DisplayName = "Buoyancy"))
 class BUOYANCY_API UBuoyancyRuntimeSettings : public UDeveloperSettings
 {
 	GENERATED_BODY()
@@ -51,7 +51,6 @@ public:
 #endif // WITH_EDITOR
 
 #if WITH_EDITORONLY_DATA
-	// Delegate called whenever the curve data is updated
 	DECLARE_MULTICAST_DELEGATE_TwoParams(FOnUpdateSettings, const UBuoyancyRuntimeSettings* /*Settings*/, EPropertyChangeType::Type /*ChangeType*/);
 	static FOnUpdateSettings OnSettingsChange;
 #endif
