In ComputeSubmergedVolume move water leaf node bounds calculations out of the innermost lambda so that they don't have to be recalculated multiple times

diff --git a/Engine/Plugins/Experimental/Buoyancy/Source/Runtime/Private/BuoyancyAlgorithms.cpp b/Engine/Plugins/Experimental/Buoyancy/Source/Runtime/Private/BuoyancyAlgorithms.cpp
--- a/Engine/Plugins/Experimental/Buoyancy/Source/Runtime/Private/BuoyancyAlgorithms.cpp
+++ b/Engine/Plugins/Experimental/Buoyancy/Source/Runtime/Private/BuoyancyAlgorithms.cpp
@@ -144,6 +144,11 @@ namespace BuoyancyAlgorithms
 			const int32 ShapeIndexA = (ShapeInstancesA.IsValidIndex(RootObjectIndexA)) ? RootObjectIndexA : 0;
 			const FShapeInstance* ShapeInstanceA = ShapeInstancesA[ShapeIndexA].Get();
 
+			// Get the world-space bounds of shape A
+			const FRigidTransform3 ShapeWorldTransformA = RelativeTransformA * ParticleWorldTransformA;
+			const FAABB3 BoxAInA = ImplicitA->BoundingBox();
+			const FAABB3 ShapeWorldBoundsA = BoxAInA.TransformedAABB(ShapeWorldTransformA);
+
 			//
 			// TODO: Dig down into ImplicitB's leaves based on its type, so that
 			// we don't end up going down to the tiniest particles of a GC.
@@ -154,6 +159,7 @@ namespace BuoyancyAlgorithms
 				[ParticleA, ImplicitA, ShapeInstanceA, LeafObjectIndexA,
 				&ParticleWorldTransformA, &RelativeTransformA,
 				ParticleIndexB, &ParticleWorldTransformB,
+				&ShapeWorldTransformA, &BoxAInA, &ShapeWorldBoundsA,
 				&NumSubdivisions, &MinVolume, &SubmergedShapes, &SubmergedVol, &SubmergedCoM]
 				(const FImplicitObject* ImplicitB, const FRigidTransform3& RelativeTransformB, const int32 RootObjectIndexB, const int32 ObjectIndexB, const int32 LeafObjectIndexB)
 			{
@@ -173,17 +179,12 @@ namespace BuoyancyAlgorithms
 					return;
 				}
 
-				// Get shape world transforms
-				const FRigidTransform3 ShapeWorldTransformA = RelativeTransformA * ParticleWorldTransformA;
+				// Get shape world-space bounds of shape B
 				const FRigidTransform3 ShapeWorldTransformB = RelativeTransformB * ParticleWorldTransformB;
-
-				// Get local bounds of objects
-				const FAABB3 BoxAInA = ImplicitA->BoundingBox();
 				const FAABB3 BoxBInB = ImplicitB->BoundingBox();
-
-				// AABB intersect check
-				const FAABB3 ShapeWorldBoundsA = BoxAInA.TransformedAABB(ShapeWorldTransformA);
 				const FAABB3 ShapeWorldBoundsB = BoxBInB.TransformedAABB(ShapeWorldTransformB);
+
+				// World-bounds AABB intersect check
 				if (!ShapeWorldBoundsA.Intersects(ShapeWorldBoundsB)) { return; }
 
 				// OOBB vs AABB intersect checks
