Add water drag factor to buoyancy runtime settings

diff --git a/Engine/Plugins/Experimental/Buoyancy/Source/Runtime/Private/BuoyancySubsystem.cpp b/Engine/Plugins/Experimental/Buoyancy/Source/Runtime/Private/BuoyancySubsystem.cpp
--- a/Engine/Plugins/Experimental/Buoyancy/Source/Runtime/Private/BuoyancySubsystem.cpp
+++ b/Engine/Plugins/Experimental/Buoyancy/Source/Runtime/Private/BuoyancySubsystem.cpp
@@ -112,6 +112,7 @@ void UBuoyancySubsystem::ApplyRuntimeSettings(const UBuoyancyRuntimeSettings* In
 	// Runtime settings presents water density in g/cm^3, but we want it in kg/cm^3
 	// so introduce a factor of 10^-3 here.
 	BuoyancySettings.WaterDensity = Chaos::GCm3ToKgCm3(InSettings->WaterDensity);
+	BuoyancySettings.WaterDrag = InSettings->WaterDrag;
 	BuoyancySettings.WaterCollisionChannel = InSettings->CollisionChannelForWaterObjects;
 	BuoyancySettings.bKeepAwake = InSettings->bKeepFloatingObjectsAwake;
 	BuoyancySettings.MaxNumBoundsSubdivisions = InSettings->MaxNumBoundsSubdivisions;
diff --git a/Engine/Plugins/Experimental/Buoyancy/Source/Runtime/Public/BuoyancyRuntimeSettings.h b/Engine/Plugins/Experimental/Buoyancy/Source/Runtime/Public/BuoyancyRuntimeSettings.h
--- a/Engine/Plugins/Experimental/Buoyancy/Source/Runtime/Public/BuoyancyRuntimeSettings.h
+++ b/Engine/Plugins/Experimental/Buoyancy/Source/Runtime/Public/BuoyancyRuntimeSettings.h
@@ -32,6 +32,11 @@ public:
 	UPROPERTY(EditAnywhere, config, Category = Buoyancy, Meta = (ClampMin = 0, ForceUnits = "g/cm3"))
 	float WaterDensity = 1.f;
 
+	/** Drag factor for submerged objects. This unitless number approximates the idea of viscosity,
+	    but internally functions identically to the "ether drag" concept. */
+	UPROPERTY(EditAnywhere, config, Category = Buoyancy, Meta = (ClampMin = 0))
+	float WaterDrag = 1.f;
+
 	/** Maximum number of times that a buoyancy bounds can be split into 8 */
 	UPROPERTY(EditAnywhere, config, Category = Buoyancy, Meta = (ClampMin = 0))
 	int32 MaxNumBoundsSubdivisions = 2;
