Use VisitLeafObjects in BuoyancyAlgorithms to compute a volume ratio and get a better approximation of submerged volume.

diff --git a/Engine/Plugins/Experimental/Buoyancy/Source/Runtime/Private/BuoyancyAlgorithms.cpp b/Engine/Plugins/Experimental/Buoyancy/Source/Runtime/Private/BuoyancyAlgorithms.cpp
--- a/Engine/Plugins/Experimental/Buoyancy/Source/Runtime/Private/BuoyancyAlgorithms.cpp
+++ b/Engine/Plugins/Experimental/Buoyancy/Source/Runtime/Private/BuoyancyAlgorithms.cpp
@@ -6,7 +6,11 @@
 #include "Engine/Engine.h"
 #include "Chaos/ParticleHandle.h"
 #include "Chaos/DebugDrawQueue.h"
+#include "Chaos/CastingUtilities.h"
 #include "Chaos/Utilities.h"
+#include "Chaos/PBDRigidsEvolutionGBF.h"
+#include "Chaos/Collision/CollisionFilter.h"
+#include "Chaos/Collision/CollisionUtil.h"
 
 //
 // CVars
@@ -107,6 +111,114 @@ namespace BuoyancyAlgorithms
 {
 	using namespace Chaos;
 
+	FRealSingle ComputeParticleVolume(const FPBDRigidsEvolutionGBF* Evolution, const FGeometryParticleHandle* Particle)
+	{
+		const FPBDRigidParticleHandle* Rigid = Particle->CastToRigidParticle();
+		if (Rigid == nullptr)
+		{
+			return -1.f;
+		}
+
+		const FChaosPhysicsMaterial* ParticleMaterial = Evolution->GetFirstClusteredPhysicsMaterial(Particle);
+		if (ParticleMaterial == nullptr)
+		{
+			return -1.f;
+		}
+
+		// Get the material density from the submerged particle's material and use that
+		// in conjunction with its mass to compute its effective total volume.
+		//
+		// Use this as the upper bound for submerged volume, since the voxelized
+		// submerged shape bounds will likely have overestimated the "true" volume
+		// of the object.
+		//
+		// NOTE: This is using the density of the material of the FIRST shape on the
+		// object, whatever it is. If for example the particle is a cluster union of
+		// GCs of totally different types, this might be an incorrect volume.
+		//
+		// However, the volumes or masses of each "true" shape are not accessible to
+		// us, so at the moment this is nearly the best estimate we'll be able to get.
+		const FRealSingle ParticleDensity = Chaos::GCm3ToKgCm3(ParticleMaterial->Density);
+		const FRealSingle ParticleMass = Rigid->M();
+		const FRealSingle ParticleVol
+			= ParticleDensity > UE_SMALL_NUMBER
+			? ParticleMass / ParticleDensity
+			: 0.f;
+		return ParticleVol;
+	}
+
+	FRealSingle ComputeShapeVolume(const FGeometryParticleHandle* Particle)
+	{
+		if (Particle == nullptr)
+		{
+			return -1.f;
+		}
+
+		const FImplicitObject* ImplicitObject = Particle->GetGeometry();
+		if (ImplicitObject == nullptr)
+		{
+			return -1.f;
+		}
+
+		const FShapeInstanceArray& ShapeInstances = Particle->ShapeInstances();
+		if (ShapeInstances.Num() == 0)
+		{
+			return -1.f;
+		}
+
+		// Loop over every leaf object and sum up the volume of each of their bounds
+		// to get an upper limit on the submerged volume that can be reported by
+		// ComputeSubmergedVolume.
+		FRealSingle ShapeVol = 0.f;
+		Utilities::VisitConcreteObjects(*ImplicitObject,
+		[Particle, &ShapeInstances, &ShapeVol](const auto& Geom, int32 ShapeIndex)
+		{
+			const EImplicitObjectType ShapeType = Private::GetImplicitCollisionType(Particle, &Geom);
+			if (DoCollide(ShapeType, ShapeInstances[ShapeIndex].Get()))
+			{
+				ShapeVol += Geom.BoundingBox().GetVolume();
+			}
+		});
+
+		return ShapeVol;
+	}
+
+	bool ComputeSubmergedVolume(const FPBDRigidsEvolutionGBF* Evolution, const FGeometryParticleHandle* ParticleA, const FGeometryParticleHandle* ParticleB, int32 NumSubdivisions, float MinVolume, TSparseArray<TBitArray<>>& SubmergedShapes, float& SubmergedVol, FVec3& SubmergedCoM)
+	{
+		if (ComputeSubmergedVolume(ParticleA, ParticleB, NumSubdivisions, MinVolume, SubmergedShapes, SubmergedVol, SubmergedCoM))
+		{
+			// Get submerged object's "particle" volume and "shape" volume.
+			//
+			// The particle volume is the theoretical volume of the particle,
+			// derived from its mass and density.
+			//
+			// The shape volume is the volume of all shape bounds which can
+			// possibly count as submerged volumes.
+			const FRealSingle ParticleVolB = ComputeParticleVolume(Evolution, ParticleB);
+			const FRealSingle ShapeVolB = ComputeShapeVolume(ParticleB);
+
+			// If the submerged vol somehow exceeded the max shape vol, clamp it
+			if (!ensureMsgf(SubmergedVol - ShapeVolB < UE_SMALL_NUMBER, TEXT("BuoyancyAlgorithms::ComputeSubmergedVolume: Somehow submerged volume exceeded theoretical maximum. Check the ComputeShapeVolume algorithm")))
+			{
+				SubmergedVol = ShapeVolB;
+			}
+
+			// Adjust the output volume based on the ratio of the material volume and the shape volume.
+			// We expect the shape volume to have overestimated the submerged volume for most shapes,
+			// especially those which are hollow.
+			if (ParticleVolB > UE_SMALL_NUMBER &&
+				ParticleVolB < ShapeVolB)
+			{
+				const float VolRatio = ParticleVolB / ShapeVolB;
+				SubmergedVol *= VolRatio;
+			}
+
+			return true;
+		}
+
+		return false;
+	}
+
 	bool ComputeSubmergedVolume(const FGeometryParticleHandle* ParticleA, const FGeometryParticleHandle* ParticleB, const int32 NumSubdivisions, const float MinVolume, TSparseArray<TBitArray<>>& SubmergedShapes, float& SubmergedVol, FVec3& SubmergedCoM)
 	{
 		SCOPE_CYCLE_COUNTER(STAT_BuoyancyAlgorithms_ComputeSubmergedVolume)
@@ -135,7 +247,7 @@ namespace BuoyancyAlgorithms
 		// object hierarchy of ParticleB.
 		RootImplicitA->VisitLeafObjects(
 			[ParticleA, &ShapeInstancesA, &ParticleWorldTransformA,
-			ParticleIndexB, &ShapeInstancesB, RootImplicitB, &ParticleWorldTransformB, &ParticleTransformAToB,
+			ParticleB, ParticleIndexB, &ShapeInstancesB, RootImplicitB, &ParticleWorldTransformB, &ParticleTransformAToB,
 			&NumSubdivisions, &MinVolume, &SubmergedShapes, &SubmergedVol, &SubmergedCoM]
 			(const FImplicitObject* ImplicitA, const FRigidTransform3& RelativeTransformA, const int32 RootObjectIndexA, const int32 ObjectIndex, const int32 LeafObjectIndexA)
 		{
@@ -143,6 +255,7 @@ namespace BuoyancyAlgorithms
 			const FAABB3 ShapeBoundsAInB = RelativeBoundsA.TransformedAABB(ParticleTransformAToB);
 			const int32 ShapeIndexA = (ShapeInstancesA.IsValidIndex(RootObjectIndexA)) ? RootObjectIndexA : 0;
 			const FShapeInstance* ShapeInstanceA = ShapeInstancesA[ShapeIndexA].Get();
+			const EImplicitObjectType ShapeTypeA = Private::GetImplicitCollisionType(ParticleA, ImplicitA);
 
 			// Get the world-space bounds of shape A
 			const FRigidTransform3 ShapeWorldTransformA = RelativeTransformA * ParticleWorldTransformA;
@@ -156,18 +269,20 @@ namespace BuoyancyAlgorithms
 
 			// Detect collisions between ImplicitA and Implicit Hierarchy of ParticleB
 			RootImplicitB->VisitOverlappingLeafObjects(ShapeBoundsAInB,
-				[ParticleA, ImplicitA, ShapeInstanceA, LeafObjectIndexA,
+				[ParticleA, ImplicitA, ShapeInstanceA, LeafObjectIndexA, ShapeTypeA,
 				&ParticleWorldTransformA, &RelativeTransformA,
-				ParticleIndexB, &ParticleWorldTransformB,
+				ParticleB, ParticleIndexB, &ShapeInstancesB, &ParticleWorldTransformB,
 				&ShapeWorldTransformA, &BoxAInA, &ShapeWorldBoundsA,
 				&NumSubdivisions, &MinVolume, &SubmergedShapes, &SubmergedVol, &SubmergedCoM]
 				(const FImplicitObject* ImplicitB, const FRigidTransform3& RelativeTransformB, const int32 RootObjectIndexB, const int32 ObjectIndexB, const int32 LeafObjectIndexB)
 			{
-				// If this shape doesn't have collision then skip it
-				//
-				// NOTE: Do we possibly need to do this in a more sophisticated way, like with actual
-				// collision filtering?
-				if (ImplicitB->GetDoCollide() == false)
+				// Get shape instance data for shape B
+				const int32 ShapeIndexB = (ShapeInstancesB.IsValidIndex(RootObjectIndexB)) ? RootObjectIndexB : 0;
+				const FShapeInstance* ShapeInstanceB = ShapeInstancesB[ShapeIndexB].Get();
+				const EImplicitObjectType ShapeTypeB = Private::GetImplicitCollisionType(ParticleB, ImplicitB);
+
+				// If this shape pair doesn't pass a narrow phase test then skip it
+				if (!ShapePairNarrowPhaseFilter(ShapeTypeA, ShapeInstanceA, ShapeTypeB, ShapeInstanceB))
 				{
 					return;
 				}
diff --git a/Engine/Plugins/Experimental/Buoyancy/Source/Runtime/Private/BuoyancySubsystem.cpp b/Engine/Plugins/Experimental/Buoyancy/Source/Runtime/Private/BuoyancySubsystem.cpp
--- a/Engine/Plugins/Experimental/Buoyancy/Source/Runtime/Private/BuoyancySubsystem.cpp
+++ b/Engine/Plugins/Experimental/Buoyancy/Source/Runtime/Private/BuoyancySubsystem.cpp
@@ -272,7 +272,7 @@ void FBuoyancySubsystemSimCallback::OnMidPhaseModification_Internal(Chaos::FMidP
 	// our target collision channel. It would be nice if it were possible instead
 	// to get a list of water body particles and loop over only midphases which
 	// involve them.
-	MidPhaseAccessor.VisitMidPhases([this](Chaos::FMidPhaseModifier& MidPhase)
+	MidPhaseAccessor.VisitMidPhases([this, Evolution](Chaos::FMidPhaseModifier& MidPhase)
 	{
 		SCOPE_CYCLE_COUNTER(STAT_BuoyancySubsystem_VisitMidphases)
 
@@ -311,7 +311,7 @@ void FBuoyancySubsystemSimCallback::OnMidPhaseModification_Internal(Chaos::FMidP
 		// Compute submerged volume and CoM
 		float SubmergedVol;
 		FVec3 SubmergedCoM;
-		if (BuoyancyAlgorithms::ComputeSubmergedVolume(WaterParticle, RigidParticle, BuoyancySettings->MaxNumBoundsSubdivisions, BuoyancySettings->MinBoundsSubdivisionVol, SubmergedShapes, SubmergedVol, SubmergedCoM))
+		if (BuoyancyAlgorithms::ComputeSubmergedVolume(Evolution, WaterParticle, RigidParticle, BuoyancySettings->MaxNumBoundsSubdivisions, BuoyancySettings->MinBoundsSubdivisionVol, SubmergedShapes, SubmergedVol, SubmergedCoM))
 		{
 			SCOPE_CYCLE_COUNTER(STAT_BuoyancySubsystem_ComputeBuoyantForces)
 
@@ -349,31 +349,6 @@ void FBuoyancySubsystemSimCallback::OnMidPhaseModification_Internal(Chaos::FMidP
 	// Apply all buoyant forces
 	for (const FSubmersion& Submersion : Submersions)
 	{
-		// Get the material density from the submerged particle's material and use that
-		// in conjunction with its mass to compute its effective total volume.
-		//
-		// Use this as the upper bound for submerged volume, since the voxelized
-		// submerged shape bounds will likely have overestimated the "true" volume
-		// of the object.
-		//
-		// NOTE: This is using the density of the material of the FIRST shape on the
-		// object, whatever it is. If for example the particle is a cluster union of
-		// GCs of totally different types, this might be an incorrect volume.
-		//
-		// However, the volumes or masses of each "true" shape are not accessible to
-		// us, so at the moment this is nearly the best estimate we'll be able to get.
-		FRealSingle SubmergedVol = Submersion.Vol;
-		if (const FChaosPhysicsMaterial* ParticleMaterial = Evolution->GetFirstClusteredPhysicsMaterial(Submersion.Particle))
-		{
-			const FRealSingle ParticleDensity = Chaos::GCm3ToKgCm3(ParticleMaterial->Density);
-			const FRealSingle ParticleMass = Submersion.Particle->M();
-			const FRealSingle ParticleVol
-				= ParticleDensity > UE_SMALL_NUMBER
-				? ParticleMass / ParticleDensity
-				: 0.f;
-			SubmergedVol = FMath::Min(ParticleVol, Submersion.Vol);
-		}
-
 		// Figure out the gravity level of the particle
 		const int32 GravityGroupIndex = Submersion.Particle->GravityGroupIndex();
 		const FVec3 GravityAccel
@@ -384,7 +359,7 @@ void FBuoyancySubsystemSimCallback::OnMidPhaseModification_Internal(Chaos::FMidP
 		// Compute delta linear and angular velocities due to buoyancy. If they're big enough to
 		// matter, apply them
 		FVec3 DeltaV, DeltaW;
-		if (BuoyancyAlgorithms::ComputeBuoyantForce(Submersion.Particle, DeltaSeconds, BuoyancySettings->WaterDensity, BuoyancySettings->WaterDrag, GravityAccel, Submersion.CoM, SubmergedVol, DeltaV, DeltaW))
+		if (BuoyancyAlgorithms::ComputeBuoyantForce(Submersion.Particle, DeltaSeconds, BuoyancySettings->WaterDensity, BuoyancySettings->WaterDrag, GravityAccel, Submersion.CoM, Submersion.Vol, DeltaV, DeltaW))
 		{
 			// Clamp delta velocities
 			DeltaV = DeltaV.GetClampedToSize(0.f, BuoyancySettings->MaxDeltaV);
diff --git a/Engine/Plugins/Experimental/Buoyancy/Source/Runtime/Public/BuoyancyAlgorithms.h b/Engine/Plugins/Experimental/Buoyancy/Source/Runtime/Public/BuoyancyAlgorithms.h
--- a/Engine/Plugins/Experimental/Buoyancy/Source/Runtime/Public/BuoyancyAlgorithms.h
+++ b/Engine/Plugins/Experimental/Buoyancy/Source/Runtime/Public/BuoyancyAlgorithms.h
@@ -7,6 +7,11 @@
 #include "Chaos/ParticleHandleFwd.h"
 #include "Chaos/AABB.h"
 
+namespace Chaos
+{
+	class FPBDRigidsEvolutionGBF;
+}
+
 namespace BuoyancyAlgorithms
 {
 	using namespace Chaos;
@@ -24,6 +29,18 @@ namespace BuoyancyAlgorithms
 		FVec3 SubmergedCoM;
 	};
 
+	// Compute the effective volume of an entire particle based on its material
+	// density and mass.
+	FRealSingle ComputeParticleVolume(const FPBDRigidsEvolutionGBF* Evolution, const FGeometryParticleHandle* Particle);
+
+	// Compute the effective volume of a shape. This method must reflect the
+	// maximum possible output value of the non-scaled ComputeSubmergedVolume.
+	FRealSingle ComputeShapeVolume(const FGeometryParticleHandle* Particle);
+
+	// Compute an approximate volume and center of mass of particle B submerged in particle A,
+	// adjusting for the volume of the object based on the material density and mass of the object
+	bool ComputeSubmergedVolume(const FPBDRigidsEvolutionGBF* Evolution, const FGeometryParticleHandle* ParticleA, const FGeometryParticleHandle* ParticleB, int32 NumSubdivisions, float MinVolume, TSparseArray<TBitArray<>>& SubmergedShapes, float& SubmergedVol, FVec3& SubmergedCoM);
+
 	// Compute an approximate volume and center of mass of particle B submerged in particle A
 	bool ComputeSubmergedVolume(const FGeometryParticleHandle* ParticleA, const FGeometryParticleHandle* ParticleB, int32 NumSubdivisions, float MinVolume, TSparseArray<TBitArray<>>& SubmergedShapes, float& SubmergedVol, FVec3& SubmergedCoM);
 
