Use VisitLeafObjects rather than VisitConcreteObjects for now because the utility crashes on a bad cast in some cases.

diff --git a/Engine/Plugins/Experimental/Buoyancy/Source/Runtime/Private/BuoyancyAlgorithms.cpp b/Engine/Plugins/Experimental/Buoyancy/Source/Runtime/Private/BuoyancyAlgorithms.cpp
--- a/Engine/Plugins/Experimental/Buoyancy/Source/Runtime/Private/BuoyancyAlgorithms.cpp
+++ b/Engine/Plugins/Experimental/Buoyancy/Source/Runtime/Private/BuoyancyAlgorithms.cpp
@@ -169,6 +169,26 @@ namespace BuoyancyAlgorithms
 		// Loop over every leaf object and sum up the volume of each of their bounds
 		// to get an upper limit on the submerged volume that can be reported by
 		// ComputeSubmergedVolume.
+		FRealSingle ShapeVol = 0.f;
+		ImplicitObject->VisitLeafObjects(
+			[Particle, &ShapeInstances, &ShapeVol]
+			(const FImplicitObject* InnerImplicitObject, const FRigidTransform3&, const int32 RootObjectIndex, const int32, const int32)
+		{
+			const int32 ShapeIndex = ShapeInstances.IsValidIndex(RootObjectIndex) ? RootObjectIndex : 0;
+			const EImplicitObjectType ShapeType = Private::GetImplicitCollisionType(Particle, InnerImplicitObject);
+			if (DoCollide(ShapeType, ShapeInstances[ShapeIndex].Get()))
+			{
+				Utilities::CastHelper(*InnerImplicitObject, [&ShapeVol](const auto& Geom)
+				{
+					ShapeVol += Geom.BoundingBox().GetVolume();
+				});
+			}
+		});
+
+		/*
+		NOTE:	This version is more performant, but visits invalid leaf objects and crashes
+				because of the utility cast function.
+		
 		FRealSingle ShapeVol = 0.f;
 		Utilities::VisitConcreteObjects(*ImplicitObject,
 		[Particle, &ShapeInstances, &ShapeVol](const auto& Geom, int32 ShapeIndex)
@@ -179,6 +199,7 @@ namespace BuoyancyAlgorithms
 				ShapeVol += Geom.BoundingBox().GetVolume();
 			}
 		});
+		*/
 
 		return ShapeVol;
 	}
@@ -198,7 +219,7 @@ namespace BuoyancyAlgorithms
 			const FRealSingle ShapeVolB = ComputeShapeVolume(ParticleB);
 
 			// If the submerged vol somehow exceeded the max shape vol, clamp it
-			if (!ensureMsgf(SubmergedVol - ShapeVolB < UE_SMALL_NUMBER, TEXT("BuoyancyAlgorithms::ComputeSubmergedVolume: Somehow submerged volume exceeded theoretical maximum. Check the ComputeShapeVolume algorithm")))
+			if (SubmergedVol - ShapeVolB > UE_SMALL_NUMBER)
 			{
 				SubmergedVol = ShapeVolB;
 			}
