Beginnings of buoyancy callbacks for touching surface of water

diff --git a/Engine/Plugins/Experimental/Buoyancy/Source/Runtime/Private/BuoyancyAlgorithms.cpp b/Engine/Plugins/Experimental/Buoyancy/Source/Runtime/Private/BuoyancyAlgorithms.cpp
--- a/Engine/Plugins/Experimental/Buoyancy/Source/Runtime/Private/BuoyancyAlgorithms.cpp
+++ b/Engine/Plugins/Experimental/Buoyancy/Source/Runtime/Private/BuoyancyAlgorithms.cpp
@@ -204,7 +204,7 @@ namespace BuoyancyAlgorithms
 		return ShapeVol;
 	}
 
-	bool ComputeSubmergedVolume(const FPBDRigidsEvolutionGBF* Evolution, const FGeometryParticleHandle* ParticleA, const FGeometryParticleHandle* ParticleB, int32 NumSubdivisions, float MinVolume, TSparseArray<TBitArray<>>& SubmergedShapes, float& SubmergedVol, FVec3& SubmergedCoM)
+	bool ComputeSubmergedVolume(const FPBDRigidsEvolutionGBF* Evolution, const FGeometryParticleHandle* ParticleA, const FGeometryParticleHandle* ParticleB, int32 NumSubdivisions, float MinVolume, TSparseArray<TBitArray<>>& SubmergedShapes, float& SubmergedVol, FVec3& SubmergedCoM, float& TotalVol)
 	{
 		if (ComputeSubmergedVolume(ParticleA, ParticleB, NumSubdivisions, MinVolume, SubmergedShapes, SubmergedVol, SubmergedCoM))
 		{
@@ -217,6 +217,7 @@ namespace BuoyancyAlgorithms
 			// possibly count as submerged volumes.
 			const FRealSingle ParticleVolB = ComputeParticleVolume(Evolution, ParticleB);
 			const FRealSingle ShapeVolB = ComputeShapeVolume(ParticleB);
+			TotalVol = ParticleVolB;
 
 			// If the submerged vol somehow exceeded the max shape vol, clamp it
 			if (SubmergedVol - ShapeVolB > UE_SMALL_NUMBER)
diff --git a/Engine/Plugins/Experimental/Buoyancy/Source/Runtime/Private/BuoyancySubsystem.cpp b/Engine/Plugins/Experimental/Buoyancy/Source/Runtime/Private/BuoyancySubsystem.cpp
--- a/Engine/Plugins/Experimental/Buoyancy/Source/Runtime/Private/BuoyancySubsystem.cpp
+++ b/Engine/Plugins/Experimental/Buoyancy/Source/Runtime/Private/BuoyancySubsystem.cpp
@@ -4,14 +4,19 @@
 #include "BuoyancyAlgorithms.h"
 #include "BuoyancyStats.h"
 #include "BuoyancyRuntimeSettings.h"
+#include "BuoyancyEventInterface.h"
 #include "Engine/World.h"
 #include "Physics/Experimental/PhysScene_Chaos.h"
 #include "Physics/PhysicsFiltering.h"
 #include "Chaos/MidPhaseModification.h"
 #include "Chaos/MassProperties.h"
+#include "DrawDebugHelpers.h"
+#include "Physics/Experimental/PhysScene_Chaos.h"
+#include "WaterBodyActor.h"
 #include "PBDRigidsSolver.h"			// Only needed to get perparticlegravity :(
 #include "Chaos/PBDRigidsEvolutionGBF.h"// Only needed to get perparticlegravity :(
 #include "Chaos/PerParticleGravity.h"	// Needed in order to determine force of gravity on each particle
+#include "Chaos/DebugDrawQueue.h"
 
 //
 // CVars
@@ -117,6 +122,16 @@ void UBuoyancySubsystem::ApplyRuntimeSettings(const UBuoyancyRuntimeSettings* In
 	BuoyancySettings.bKeepAwake = InSettings->bKeepFloatingObjectsAwake;
 	BuoyancySettings.MaxNumBoundsSubdivisions = InSettings->MaxNumBoundsSubdivisions;
 	BuoyancySettings.MinBoundsSubdivisionVol = InSettings->MinBoundsSubdivisionVol;
+	BuoyancySettings.MinVelocityForSurfaceTouchCallback = InSettings->MinVelocityForSurfaceTouchCallback;
+
+	UWorld* World = GetWorld();
+	BuoyancySettings.bSurfaceTouchCallback =
+#if WITH_EDITOR
+		InSettings->bSurfaceTouchCallbackOnServer || InSettings->bSurfaceTouchCallbackOnClient;
+#else
+		(World && World->IsNetMode(NM_DedicatedServer) && InSettings->bSurfaceTouchCallbackOnServer) ||
+		(World && World->IsNetMode(NM_Client) && InSettings->bSurfaceTouchCallbackOnClient);
+#endif
 
 	// Enable or disable
 	SetEnabled(InSettings->bBuoyancyEnabled);
@@ -124,19 +139,13 @@ void UBuoyancySubsystem::ApplyRuntimeSettings(const UBuoyancyRuntimeSettings* In
 
 void UBuoyancySubsystem::Tick(float DeltaTime)
 {
+	SCOPE_CYCLE_COUNTER(STAT_BuoyancySubsystem_Tick)
+
 	//
-	// NOTE: The entire point of this tick is to send runtime data to
-	// the physics thread which might have changed. At the initial
-	// time of this writing, the only piece of data is the water-collision-
-	// channel, which will never change once set in almost every case.
-	// 
-	// Therefore we can probably do away with this tick, and move the
-	// GetProducerInputData_External to PostInitialize and also
-	// ApplyRuntimeSettings for editor builds, and do away with the Tick
-	// entirely.
-	// 
-	// However, I'm not confident that that would work due to timing
-	// weirdness of the sim callback system, and also 
+	// The entire point of this tick is to send runtime data to
+	// the physics thread which might have changed, and to process
+	// outputs which may effect water bodies or result in
+	// callbacks.
 	//
 
 	Super::Tick(DeltaTime);
@@ -146,6 +155,18 @@ void UBuoyancySubsystem::Tick(float DeltaTime)
 		return;
 	}
 
+	UWorld* World = GetWorld();
+	if (World == nullptr)
+	{
+		return;
+	}
+
+	FPhysScene* PhysScene = World->GetPhysicsScene();
+	if (PhysScene == nullptr)
+	{
+		return;
+	}
+
 	// Only bother sending new async inputs if our buoyancy settings actually changed
 	if (bBuoyancySettingsChanged)
 	{
@@ -155,6 +176,54 @@ void UBuoyancySubsystem::Tick(float DeltaTime)
 			AsyncInput->BuoyancySettings = MakeUnique<FBuoyancySettings>(BuoyancySettings);
 		}
 	}
+
+	// Process surface-touched callbacks
+	if (BuoyancySettings.bSurfaceTouchCallback)
+	{
+		SCOPE_CYCLE_COUNTER(STAT_BuoyancySubsystem_DispatchCallbacks)
+
+		while (Chaos::TSimCallbackOutputHandle<FBuoyancySubsystemSimCallbackOutput> AsyncOutput = SimCallback->PopFutureOutputData_External())
+		{
+			for (const FBuoyancySubsystemSimCallbackOutput::FSurfaceTouch& SurfaceTouch : AsyncOutput->SurfaceTouches)
+			{
+				// Extract primitive components
+				UPrimitiveComponent* WaterComponent = PhysScene->GetOwningComponent<UPrimitiveComponent>(SurfaceTouch.WaterProxy);
+				UPrimitiveComponent* RigidComponent = PhysScene->GetOwningComponent<UPrimitiveComponent>(SurfaceTouch.RigidProxy);
+
+				// Get the parental water body component
+				AWaterBody* WaterBodyActor = WaterComponent->GetOwner<AWaterBody>();
+
+				// Dispatch the delegate
+				OnSurfaceTouched.Broadcast(
+					WaterBodyActor,
+					WaterComponent,
+					RigidComponent,
+					SurfaceTouch.Vol,
+					SurfaceTouch.CoM,
+					SurfaceTouch.Vel);
+
+				const auto DispatchEvent = [&](AActor* Actor)
+				{
+					// TODO: Actor relevancy check?
+
+					// If the actor implements the event interface, call the surface touched callback
+					if (Actor->Implements<UBuoyancyEventInterface>())
+					{
+						IBuoyancyEventInterface::Execute_OnSurfaceTouched(
+							Actor,
+							WaterBodyActor,
+							WaterComponent,
+							RigidComponent,
+							SurfaceTouch.Vol,
+							SurfaceTouch.CoM,
+							SurfaceTouch.Vel);
+					}
+				};
+				DispatchEvent(WaterBodyActor);
+				DispatchEvent(WaterComponent->GetOwner());
+			}
+		}
+	}
 }
 
 TStatId UBuoyancySubsystem::GetStatId() const
@@ -188,6 +257,11 @@ void FBuoyancySubsystemSimCallbackInput::Reset()
 	BuoyancySettings.Reset();
 }
 
+void FBuoyancySubsystemSimCallbackOutput::Reset()
+{
+	SurfaceTouches.Reset();
+}
+
 namespace
 {
 	using namespace Chaos;
@@ -267,114 +341,195 @@ void FBuoyancySubsystemSimCallback::OnMidPhaseModification_Internal(Chaos::FMidP
 	// Clear submersions and submerged shapes array, but keep their memory allocated
 	Submersions.Reset();
 	SubmergedShapes.Reset();
+	SubmersionMetaData.Reset();
 
-	// NOTE: For now we visit _every_ midphase, and check for ones which involve
-	// our target collision channel. It would be nice if it were possible instead
-	// to get a list of water body particles and loop over only midphases which
-	// involve them.
-	MidPhaseAccessor.VisitMidPhases([this, Evolution](Chaos::FMidPhaseModifier& MidPhase)
 	{
 		SCOPE_CYCLE_COUNTER(STAT_BuoyancySubsystem_VisitMidphases)
 
-		FGeometryParticleHandle* Particle0;
-		FGeometryParticleHandle* Particle1;
-		MidPhase.GetParticles(&Particle0, &Particle1);
-		if (Particle0 == nullptr || Particle1 == nullptr)
+		// NOTE: For now we visit _every_ midphase, and check for ones which involve
+		// our target collision channel. It would be nice if it were possible instead
+		// to get a list of water body particles and loop over only midphases which
+		// involve them.
+		MidPhaseAccessor.VisitMidPhases([this, Evolution](Chaos::FMidPhaseModifier& MidPhase)
 		{
-			return;
-		}
-
-		// Use collision filters to pick out which of the particles is water, if either.
-		//
-		// TODO: Instead of this, water particles could have additional userdata, or could
-		// register themselves with the callback and we could look for only the midphases
-		// generated with them.
-		const bool bWater0 = ParticleHasCollision(*Particle0, BuoyancySettings->WaterCollisionChannel);
-		const bool bWater1 = ParticleHasCollision(*Particle1, BuoyancySettings->WaterCollisionChannel);
-		if ((bWater0 || bWater1) == false)
-		{
-			return;
-		}
+			FGeometryParticleHandle* Particle0;
+			FGeometryParticleHandle* Particle1;
+			MidPhase.GetParticles(&Particle0, &Particle1);
+			if (Particle0 == nullptr || Particle1 == nullptr)
+			{
+				return;
+			}
 
-		// We want to disable all collisions with water, so disable the midphase now
-		MidPhase.Disable();
+			// Use collision filters to pick out which of the particles is water, if either.
+			//
+			// TODO: Instead of this, water particles could have additional userdata, or could
+			// register themselves with the callback and we could look for only the midphases
+			// generated with them.
+			const bool bWater0 = ParticleHasCollision(*Particle0, BuoyancySettings->WaterCollisionChannel);
+			const bool bWater1 = ParticleHasCollision(*Particle1, BuoyancySettings->WaterCollisionChannel);
+			if ((bWater0 || bWater1) == false)
+			{
+				return;
+			}
 
-		// Store the rigid particle and the water particle
-		if (bWater0) { Swap(Particle0, Particle1); }
-		FPBDRigidParticleHandle* RigidParticle = Particle0->CastToRigidParticle();
-		FGeometryParticleHandle* WaterParticle = Particle1;
-		if (RigidParticle == nullptr)
-		{
-			return;
-		}
+			// We want to disable all collisions with water, so disable the midphase now
+			MidPhase.Disable();
 
-		// Compute submerged volume and CoM
-		float SubmergedVol;
-		FVec3 SubmergedCoM;
-		if (BuoyancyAlgorithms::ComputeSubmergedVolume(Evolution, WaterParticle, RigidParticle, BuoyancySettings->MaxNumBoundsSubdivisions, BuoyancySettings->MinBoundsSubdivisionVol, SubmergedShapes, SubmergedVol, SubmergedCoM))
-		{
-			SCOPE_CYCLE_COUNTER(STAT_BuoyancySubsystem_ComputeBuoyantForces)
+			// Store the rigid particle and the water particle
+			if (bWater0) { Swap(Particle0, Particle1); }
+			FPBDRigidParticleHandle* RigidParticle = Particle0->CastToRigidParticle();
+			FGeometryParticleHandle* WaterParticle = Particle1;
+			if (RigidParticle == nullptr)
+			{
+				return;
+			}
 
-			//
-			// If any volume was submerged, take the weighted average of the centers of mass to get the
-			// approximate submerged center of mass, and apply the buoyancy force there.
-			if (SubmergedVol > SMALL_NUMBER)
+			// Compute submerged volume and CoM
+			float SubmergedVol;
+			FVec3 SubmergedCoM;
+			float TotalVol;
+			if (BuoyancyAlgorithms::ComputeSubmergedVolume(Evolution, WaterParticle, RigidParticle, BuoyancySettings->MaxNumBoundsSubdivisions, BuoyancySettings->MinBoundsSubdivisionVol, SubmergedShapes, SubmergedVol, SubmergedCoM, TotalVol))
 			{
-				const int32 RigidParticleIndex = RigidParticle->UniqueIdx().Idx;
+				SCOPE_CYCLE_COUNTER(STAT_BuoyancySubsystem_BuildSubmersions)
 
-				// If this particle was already marked submerged, add to its existing submersion
-				if (Submersions.IsValidIndex(RigidParticleIndex))
+				// If any volume was submerged, take the weighted average of the centers of mass to get the
+				// approximate submerged center of mass, and apply the buoyancy force there.
+				if (SubmergedVol > SMALL_NUMBER)
 				{
-					FSubmersion& Submersion = Submersions[RigidParticleIndex];
-					ensureMsgf(Submersion.Particle == RigidParticle, TEXT("Something went wrong - there's a particle index mismatch in the Submersions sparse array"));
+					const int32 RigidParticleIndex = RigidParticle->UniqueIdx().Idx;
+
+					// If this particle was already marked submerged, add to its existing submersion
+					if (Submersions.IsValidIndex(RigidParticleIndex))
+					{
+						FBuoyancySubmersion& Submersion = Submersions[RigidParticleIndex];
+						ensureMsgf(Submersion.Particle == RigidParticle, TEXT("Something went wrong - there's a particle index mismatch in the Submersions sparse array"));
+
+						// Sum the volumes
+						Submersion.Vol = Submersion.Vol + SubmergedVol;
+
+						// Get the weighted-average CoM
+						// NOTE: The unchecked division should be safe since we already
+						// know SubmergedVol > SMALL_NUMBER
+						Submersion.CoM = ((Submersion.CoM * Submersion.Vol) + (SubmergedCoM * SubmergedVol)) / Submersion.Vol;
+					}
+
+					// If this particle was not yet submerged, make a new submersion for it
+					else
+					{
+						Submersions.Insert(RigidParticleIndex, { RigidParticle, SubmergedVol, SubmergedCoM });
+					}
+
+					// If this is a surface touch record it for callback, if 
+					if (BuoyancySettings->bSurfaceTouchCallback && TotalVol > SubmergedVol * (1.f + UE_KINDA_SMALL_NUMBER))
+					{
+						SCOPE_CYCLE_COUNTER(STAT_BuoyancySubsystem_BuildSubmersionCallbackData)
+
+						// Proceed only if this CoM is moving fast enough to generate events
+						const FBuoyancySubmersion& Submersion = Submersions[RigidParticleIndex];
+						const FVec3 CoMDiff = Submersion.CoM - RigidParticle->XCom();
+						const FVec3 CoMVel = RigidParticle->V() + FVec3::CrossProduct(RigidParticle->W(), CoMDiff);
+						const float CoMVelSq = FVec3::DotProduct(CoMVel, CoMVel);
+						const float MinVel = BuoyancySettings->MinVelocityForSurfaceTouchCallback;
+						const float MinVelSq = MinVel * MinVel;
+
+#if ENABLE_DRAW_DEBUG
+						if (bBuoyancyDebugDraw)
+						{
+							Chaos::FDebugDrawQueue::GetInstance().DrawDebugDirectionalArrow(
+								Submersion.CoM, Submersion.CoM + CoMVel, 30.f, FColor::Yellow, false, -1, -1, 2.f);
+						}
+#endif
+
+						if (CoMVelSq > MinVelSq)
+						{
+							// If we don't have a metadata for this particle yet, add one
+							if (!SubmersionMetaData.IsValidIndex(RigidParticleIndex))
+							{
+								SubmersionMetaData.Insert(RigidParticleIndex, FBuoyancySubmersionMetaData());
+							}
+							FBuoyancySubmersionMetaData& MetaData = SubmersionMetaData[RigidParticleIndex];
+
+							// If we haven't already maxed out on water contacts, add one
+							if (MetaData.WaterContacts.Num() < FBuoyancySubmersionMetaData::MaxNumWaterContacts)
+							{
+								MetaData.WaterContacts.Add({ WaterParticle, SubmergedVol, SubmergedCoM, CoMVel });
+							}
+						}
+					}
+				}
+			}
+		});
+	}
 
-					// Sum the volumes
-					Submersion.Vol = Submersion.Vol + SubmergedVol;
+	{
+		SCOPE_CYCLE_COUNTER(STAT_BuoyancySubsystem_ApplyBuoyantForces)
 
-					// Get the weighted-average CoM
-					// NOTE: The unchecked division should be safe since we already
-					// know SubmergedVol > SMALL_NUMBER
-					Submersion.CoM = ((Submersion.CoM * Submersion.Vol) + (SubmergedCoM * SubmergedVol)) / Submersion.Vol;
-				}
+		// Apply all buoyant forces
+		for (const FBuoyancySubmersion& Submersion : Submersions)
+		{
+			// Figure out the gravity level of the particle
+			const int32 GravityGroupIndex = Submersion.Particle->GravityGroupIndex();
+			const FVec3 GravityAccel
+				= PerParticleGravity != nullptr && GravityGroupIndex != INDEX_NONE
+				? (FVec3)PerParticleGravity->GetAcceleration(GravityGroupIndex)
+				: FVec3::DownVector * 980.f; // Default to "regular" gravity
+
+			// Compute delta linear and angular velocities due to buoyancy. If they're big enough to
+			// matter, apply them
+			FVec3 DeltaV, DeltaW;
+			if (BuoyancyAlgorithms::ComputeBuoyantForce(Submersion.Particle, DeltaSeconds, BuoyancySettings->WaterDensity, BuoyancySettings->WaterDrag, GravityAccel, Submersion.CoM, Submersion.Vol, DeltaV, DeltaW))
+			{
+				// Clamp delta velocities
+				DeltaV = DeltaV.GetClampedToSize(0.f, BuoyancySettings->MaxDeltaV);
+				DeltaW = DeltaW.GetClampedToSize(0.f, BuoyancySettings->MaxDeltaW);
 
-				// If this particle was not yet submerged, make a new submersion for it
-				else
+				// Apply the deltas
+				Submersion.Particle->SetV(Submersion.Particle->V() + DeltaV);
+				Submersion.Particle->SetW(Submersion.Particle->W() + DeltaW);
+
+				// Wake up the body??
+				if (Evolution && BuoyancySettings->bKeepAwake)
 				{
-					Submersions.Insert(RigidParticleIndex, { RigidParticle, SubmergedVol, SubmergedCoM });
+					Evolution->SetParticleObjectState(Submersion.Particle, EObjectStateType::Dynamic);
 				}
 			}
 		}
-	});
+	}
+
 
-	// Apply all buoyant forces
-	for (const FSubmersion& Submersion : Submersions)
+	// Generate callback data if we're into that sort of thing
+	if (BuoyancySettings->bSurfaceTouchCallback)
 	{
-		// Figure out the gravity level of the particle
-		const int32 GravityGroupIndex = Submersion.Particle->GravityGroupIndex();
-		const FVec3 GravityAccel
-			= PerParticleGravity != nullptr && GravityGroupIndex != INDEX_NONE
-			? (FVec3)PerParticleGravity->GetAcceleration(GravityGroupIndex)
-			: FVec3::DownVector * 980.f; // Default to "regular" gravity
-
-		// Compute delta linear and angular velocities due to buoyancy. If they're big enough to
-		// matter, apply them
-		FVec3 DeltaV, DeltaW;
-		if (BuoyancyAlgorithms::ComputeBuoyantForce(Submersion.Particle, DeltaSeconds, BuoyancySettings->WaterDensity, BuoyancySettings->WaterDrag, GravityAccel, Submersion.CoM, Submersion.Vol, DeltaV, DeltaW))
+		SCOPE_CYCLE_COUNTER(STAT_BuoyancySubsystem_ProduceSurfaceTouches)
+
+		// Get the async output struct to write to
+		FBuoyancySubsystemSimCallbackOutput& Output = GetProducerOutputData_Internal();
+
+		// Make sure we have metadata for each submersion
+		//ensureAlwaysMsgf(Submersions.Num() == SubmersionMetaData.Num());
+
+		// Process every surface touch and queue up some of them to return
+		// to game thread for callback dispatch
+		Output.SurfaceTouches.Reserve(SubmersionMetaData.Num() * FBuoyancySubmersionMetaData::MaxNumWaterContacts);
+		for (auto Iter = SubmersionMetaData.CreateIterator(); Iter; ++Iter)
 		{
-			// Clamp delta velocities
-			DeltaV = DeltaV.GetClampedToSize(0.f, BuoyancySettings->MaxDeltaV);
-			DeltaW = DeltaW.GetClampedToSize(0.f, BuoyancySettings->MaxDeltaW);
+			//
+			// TODO: Check to see if this is touching surface
+			//
 
-			// Apply the deltas
-			Submersion.Particle->SetV(Submersion.Particle->V() + DeltaV);
-			Submersion.Particle->SetW(Submersion.Particle->W() + DeltaW);
+			const FBuoyancySubmersionMetaData& MetaData = *Iter;
+			const FBuoyancySubmersion& Submersion = Submersions[Iter.GetIndex()];
 
-			// Wake up the body??
-			if (Evolution && BuoyancySettings->bKeepAwake)
+			for (const FBuoyancySubmersionMetaData::FWaterContact& WaterContact : MetaData.WaterContacts)
 			{
-				Evolution->SetParticleObjectState(Submersion.Particle, EObjectStateType::Dynamic);
+				Output.SurfaceTouches.Add({
+					Submersion.Particle->PhysicsProxy(),
+					WaterContact.Water->PhysicsProxy(),
+					WaterContact.Vol,
+					WaterContact.CoM,
+					WaterContact.Vel
+				});
 			}
 		}
-
 	}
 }
diff --git a/Engine/Plugins/Experimental/Buoyancy/Source/Runtime/Public/BuoyancyAlgorithms.h b/Engine/Plugins/Experimental/Buoyancy/Source/Runtime/Public/BuoyancyAlgorithms.h
--- a/Engine/Plugins/Experimental/Buoyancy/Source/Runtime/Public/BuoyancyAlgorithms.h
+++ b/Engine/Plugins/Experimental/Buoyancy/Source/Runtime/Public/BuoyancyAlgorithms.h
@@ -39,7 +39,7 @@ namespace BuoyancyAlgorithms
 
 	// Compute an approximate volume and center of mass of particle B submerged in particle A,
 	// adjusting for the volume of the object based on the material density and mass of the object
-	bool ComputeSubmergedVolume(const FPBDRigidsEvolutionGBF* Evolution, const FGeometryParticleHandle* ParticleA, const FGeometryParticleHandle* ParticleB, int32 NumSubdivisions, float MinVolume, TSparseArray<TBitArray<>>& SubmergedShapes, float& SubmergedVol, FVec3& SubmergedCoM);
+	bool ComputeSubmergedVolume(const FPBDRigidsEvolutionGBF* Evolution, const FGeometryParticleHandle* ParticleA, const FGeometryParticleHandle* ParticleB, int32 NumSubdivisions, float MinVolume, TSparseArray<TBitArray<>>& SubmergedShapes, float& SubmergedVol, FVec3& SubmergedCoM, float& TotalVol);
 
 	// Compute an approximate volume and center of mass of particle B submerged in particle A
 	bool ComputeSubmergedVolume(const FGeometryParticleHandle* ParticleA, const FGeometryParticleHandle* ParticleB, int32 NumSubdivisions, float MinVolume, TSparseArray<TBitArray<>>& SubmergedShapes, float& SubmergedVol, FVec3& SubmergedCoM);
diff --git a/Engine/Plugins/Experimental/Buoyancy/Source/Runtime/Public/BuoyancyEventInterface.h b/Engine/Plugins/Experimental/Buoyancy/Source/Runtime/Public/BuoyancyEventInterface.h
new file mode 100644
--- /dev/null
+++ b/Engine/Plugins/Experimental/Buoyancy/Source/Runtime/Public/BuoyancyEventInterface.h
@@ -0,0 +1,28 @@
+// Copyright Epic Games, Inc. All Rights Reserved.
+
+#pragma once
+#include "CoreMinimal.h"
+#include "BuoyancyEventInterface.generated.h"
+
+UINTERFACE(Blueprintable)
+class BUOYANCY_API UBuoyancyEventInterface : public UInterface
+{
+	GENERATED_BODY()
+};
+
+class BUOYANCY_API IBuoyancyEventInterface
+{
+	GENERATED_BODY()
+
+public:
+
+	UFUNCTION(BlueprintImplementableEvent, Category = WaterBody)
+	void  OnSurfaceTouched(
+		class AWaterBody* WaterBodyActor,
+		UPrimitiveComponent* WaterComponent,
+		UPrimitiveComponent* SubmergedComponent,
+		float SubmergedVolume,
+		const FVector& SubmergedCenterOfMass,
+		const FVector& SubmergedVelocity);
+
+};
diff --git a/Engine/Plugins/Experimental/Buoyancy/Source/Runtime/Public/BuoyancyRuntimeSettings.h b/Engine/Plugins/Experimental/Buoyancy/Source/Runtime/Public/BuoyancyRuntimeSettings.h
--- a/Engine/Plugins/Experimental/Buoyancy/Source/Runtime/Public/BuoyancyRuntimeSettings.h
+++ b/Engine/Plugins/Experimental/Buoyancy/Source/Runtime/Public/BuoyancyRuntimeSettings.h
@@ -20,25 +20,25 @@ public:
 	/** Whether or not the buoyancy subsystem should run.
 	    The subsystem has a public interface for enabling/disabling, so
 	    it's possible that a game overrides this at runtime */
-	UPROPERTY(EditAnywhere, config, Category = Buoyancy, Meta = (ClampMin = 0, ForceUnits = "g/cm3"))
+	UPROPERTY(EditAnywhere, Config, Category = Buoyancy, Meta = (ClampMin = 0, ForceUnits = "g/cm3"))
 	bool bBuoyancyEnabled = true;
 
 	/** Prevent floating particles from falling asleep by explicitly setting
 	    object state to dynamic when they are detected. */
-	UPROPERTY(EditAnywhere, config, Category = Buoyancy)
+	UPROPERTY(EditAnywhere, Config, Category = Buoyancy)
 	bool bKeepFloatingObjectsAwake = false;
 
 	/** Density of water to use in buoyancy calculations. The density of water is approximately 1g/cm3 */
-	UPROPERTY(EditAnywhere, config, Category = Buoyancy, Meta = (ClampMin = 0, ForceUnits = "g/cm3"))
+	UPROPERTY(EditAnywhere, Config, Category = Buoyancy, Meta = (ClampMin = 0, ForceUnits = "g/cm3"))
 	float WaterDensity = 1.f;
 
 	/** Drag factor for submerged objects. This unitless number approximates the idea of viscosity,
 	    but internally functions identically to the "ether drag" concept. */
-	UPROPERTY(EditAnywhere, config, Category = Buoyancy, Meta = (ClampMin = 0))
+	UPROPERTY(EditAnywhere, Config, Category = Buoyancy, Meta = (ClampMin = 0))
 	float WaterDrag = 1.f;
 
 	/** Maximum number of times that a buoyancy bounds can be split into 8 */
-	UPROPERTY(EditAnywhere, config, Category = Buoyancy, Meta = (ClampMin = 0))
+	UPROPERTY(EditAnywhere, Config, Category = Buoyancy, Meta = (ClampMin = 0))
 	int32 MaxNumBoundsSubdivisions = 2;
 
 	/** Minimum volume bounding box which can be produced by a subdivision operation */
@@ -46,9 +46,25 @@ public:
 	float MinBoundsSubdivisionVol = FMath::Pow(100.f, 3.f); // 1m^3
 
 	/** Collision channel to use for water ObjectTypes */
-	UPROPERTY(EditAnywhere, config, Category = Collision)
+	UPROPERTY(EditAnywhere, Config, Category = Collision)
 	TEnumAsByte<ECollisionChannel> CollisionChannelForWaterObjects;
 
+	/** When enabled, a delegate will be enabled which triggers when objects
+	    information for generating splashes and waves. */
+	UPROPERTY(EditAnywhere, Config, Category = Callbacks)
+	bool bSurfaceTouchCallbackOnServer = false;
+
+	/** When enabled, a delegate will be enabled which triggers when objects
+	    are touching the surface of a body of water on client. */
+	UPROPERTY(EditAnywhere, Config, Category = Callbacks)
+	bool bSurfaceTouchCallbackOnClient = true;
+
+	/** Minimum velocity of an object relative to a water required in order to trigger
+	    a submersion callback. If this is zero, then the callback will be triggered
+	    every frame for as long as any object is submerged. */
+	UPROPERTY(EditAnywhere, Config, Category = Callbacks, Meta = (ClampMin = 0, ForceUnits = "cm/s", EditCondition = "bSubmersionCallbackEnabled"))
+	float MinVelocityForSurfaceTouchCallback = 10.f;
+
 	virtual void PostInitProperties() override;
 
 #if WITH_EDITOR
diff --git a/Engine/Plugins/Experimental/Buoyancy/Source/Runtime/Public/BuoyancyStats.h b/Engine/Plugins/Experimental/Buoyancy/Source/Runtime/Public/BuoyancyStats.h
--- a/Engine/Plugins/Experimental/Buoyancy/Source/Runtime/Public/BuoyancyStats.h
+++ b/Engine/Plugins/Experimental/Buoyancy/Source/Runtime/Public/BuoyancyStats.h
@@ -10,7 +10,12 @@ DECLARE_STATS_GROUP(TEXT("Buoyancy"), STATGROUP_Buoyancy, STATCAT_Advanced);
 DECLARE_CYCLE_STAT(TEXT("Buoyancy_Subsystem_OnPreSimulate"), STAT_BuoyancySubsystem_OnPreSimulate, STATGROUP_Buoyancy);
 DECLARE_CYCLE_STAT(TEXT("Buoyancy_Subsystem_OnMidPhaseModification"), STAT_BuoyancySubsystem_OnMidPhaseModification, STATGROUP_Buoyancy);
 DECLARE_CYCLE_STAT(TEXT("Buoyancy_Subsystem_VisitMidphases"), STAT_BuoyancySubsystem_VisitMidphases, STATGROUP_Buoyancy);
-DECLARE_CYCLE_STAT(TEXT("Buoyancy_Subsystem_ComputeBuoyantForces"), STAT_BuoyancySubsystem_ComputeBuoyantForces, STATGROUP_Buoyancy);
+DECLARE_CYCLE_STAT(TEXT("Buoyancy_Subsystem_BuildSubmersions"), STAT_BuoyancySubsystem_BuildSubmersions, STATGROUP_Buoyancy);
+DECLARE_CYCLE_STAT(TEXT("Buoyancy_Subsystem_BuildSubmersionCallbackData"), STAT_BuoyancySubsystem_BuildSubmersionCallbackData, STATGROUP_Buoyancy);
+DECLARE_CYCLE_STAT(TEXT("Buoyancy_Subsystem_ApplyBuoyantForces"), STAT_BuoyancySubsystem_ApplyBuoyantForces, STATGROUP_Buoyancy);
+DECLARE_CYCLE_STAT(TEXT("Buoyancy_Subsystem_ProduceSurfaceTouches"), STAT_BuoyancySubsystem_ProduceSurfaceTouches, STATGROUP_Buoyancy);
+DECLARE_CYCLE_STAT(TEXT("Buoyancy_Subsystem_Tick"), STAT_BuoyancySubsystem_Tick, STATGROUP_Buoyancy);
+DECLARE_CYCLE_STAT(TEXT("Buoyancy_Subsystem_DispatchCallbacks"), STAT_BuoyancySubsystem_DispatchCallbacks, STATGROUP_Buoyancy);
 
 // Algorithms
 DECLARE_CYCLE_STAT(TEXT("Buoyancy_Algorithms_ComputeSubmergedVolume"), STAT_BuoyancyAlgorithms_ComputeSubmergedVolume, STATGROUP_Buoyancy);
diff --git a/Engine/Plugins/Experimental/Buoyancy/Source/Runtime/Public/BuoyancySubsystem.h b/Engine/Plugins/Experimental/Buoyancy/Source/Runtime/Public/BuoyancySubsystem.h
--- a/Engine/Plugins/Experimental/Buoyancy/Source/Runtime/Public/BuoyancySubsystem.h
+++ b/Engine/Plugins/Experimental/Buoyancy/Source/Runtime/Public/BuoyancySubsystem.h
@@ -7,6 +7,7 @@
 #include "Chaos/Framework/PhysicsProxyBase.h"
 #include "PBDRigidsSolver.h"
 #include "Engine/EngineBaseTypes.h"
+#include "WaterBodyComponent.h"
 #include "BuoyancySubsystem.generated.h"
 
 DECLARE_LOG_CATEGORY_EXTERN(LogBuoyancySubsystem, Log, All);
@@ -35,9 +36,27 @@ struct FBuoyancySettings
 	float MinBoundsSubdivisionVol = FMath::Pow(100.f, 3.f); // 1m^3
 
 	ECollisionChannel WaterCollisionChannel = ECollisionChannel::ECC_MAX;
+
+	bool bSurfaceTouchCallback = true;
+
+	float MinVelocityForSurfaceTouchCallback = 10.f;
 };
 
 
+//
+// Buoyancy Delegate Callback
+//
+
+UDELEGATE()
+DECLARE_DYNAMIC_MULTICAST_DELEGATE_SixParams(FSurfaceTouchedDelegate,
+	class AWaterBody*, WaterBodyActor,
+	UPrimitiveComponent*, WaterComponent,
+	UPrimitiveComponent*, SubmergedComponent,
+	float, SubmergedVolume,
+	const FVector&, SubmergedCenterOfMass,
+	const FVector&, SubmergedVelocity);
+
+
 //
 // Buoyancy Subsystem
 //
@@ -61,8 +80,14 @@ public:
 	bool SetEnabled(const bool bEnabled);
 
 	// Return true if subsystem is enabled and running
+	UFUNCTION()
 	bool IsEnabled() const;
 
+	// Give access to a delegate which will trigger when objects
+	// are in water.
+	UPROPERTY(BlueprintAssignable, Category = Buoyancy)
+	FSurfaceTouchedDelegate OnSurfaceTouched;
+
 protected:
 
 	// UTickableWorldSubsystem begin interface
@@ -94,6 +119,31 @@ private:
 // Buoyancy Sim Callback
 //
 
+// Metadata for submersions, used for event callbacks
+struct FBuoyancySubmersionMetaData
+{
+	struct FWaterContact
+	{
+		Chaos::FGeometryParticleHandle* Water;
+		float Vol;
+		FVector CoM;
+		FVector Vel;
+	};
+
+	// How many metadata's allowed per submerged particle
+	static constexpr int32 MaxNumWaterContacts = 3;
+	TArray<FWaterContact, TInlineAllocator<MaxNumWaterContacts>> WaterContacts;
+};
+
+
+// A minimal struct of data tracking all the submersions in a frame.
+struct FBuoyancySubmersion
+{
+	Chaos::FPBDRigidParticleHandle* Particle;
+	float Vol;
+	FVector CoM;
+};
+
 struct FBuoyancySubsystemSimCallbackInput : public Chaos::FSimCallbackInput
 {
 	// Here we use a unique ptr so that it is possible to provide an async
@@ -105,11 +155,27 @@ struct FBuoyancySubsystemSimCallbackInput : public Chaos::FSimCallbackInput
 	void Reset();
 };
 
+struct FBuoyancySubsystemSimCallbackOutput : public Chaos::FSimCallbackOutput
+{
+	struct FSurfaceTouch
+	{
+		IPhysicsProxyBase* RigidProxy;
+		IPhysicsProxyBase* WaterProxy;
+		float Vol;
+		FVector CoM;
+		FVector Vel;
+	};
+
+	TArray<FSurfaceTouch> SurfaceTouches;
+
+	void Reset();
+};
+
 // NOTE: The Presimulate option is only needed for proper registry with the solver.
 //       We don't actually need (or want!) a presimulate tick.
 class FBuoyancySubsystemSimCallback : public Chaos::TSimCallbackObject<
 	FBuoyancySubsystemSimCallbackInput,
-	Chaos::FSimCallbackNoOutput,
+	FBuoyancySubsystemSimCallbackOutput,
 	Chaos::ESimCallbackOptions::Presimulate | Chaos::ESimCallbackOptions::MidPhaseModification>
 {
 private:
@@ -122,19 +188,15 @@ private:
 	// memory that was allocated by GT to minimize copies.
 	TUniquePtr<FBuoyancySettings> BuoyancySettings;
 
-	// A minimal struct of data tracking all the submersions in a frame.
-	struct FSubmersion
-	{
-		Chaos::FPBDRigidParticleHandle* Particle;
-		float Vol;
-		Chaos::FVec3 CoM;
-	};
-
 	// This sparse array of submersion events is indexed on particle unique indices.
 	// All buoyant forces due to submersions are applied at once. It's stored as
 	// a member variable and reset every frame, to avoid reallocation of similarly
 	// sized data.
-	TSparseArray<FSubmersion> Submersions;
+	TSparseArray<FBuoyancySubmersion> Submersions;
+
+	// Another sparse array to be kept in sync with Submersions, which will contain
+	// metadata useful for event callbacks
+	TSparseArray<FBuoyancySubmersionMetaData> SubmersionMetaData;
 
 	// This is a sparse array of bit arrays representing which shapes in an object
 	// have already been accounted for when submerging an object. For example, if
