Change default min velocity for surface touch callback from 10cm/s to 1m/s

diff --git a/Engine/Plugins/Experimental/Buoyancy/Source/Runtime/Public/BuoyancyRuntimeSettings.h b/Engine/Plugins/Experimental/Buoyancy/Source/Runtime/Public/BuoyancyRuntimeSettings.h
--- a/Engine/Plugins/Experimental/Buoyancy/Source/Runtime/Public/BuoyancyRuntimeSettings.h
+++ b/Engine/Plugins/Experimental/Buoyancy/Source/Runtime/Public/BuoyancyRuntimeSettings.h
@@ -63,7 +63,7 @@ public:
 	    a submersion callback. If this is zero, then the callback will be triggered
 	    every frame for as long as any object is submerged. */
 	UPROPERTY(EditAnywhere, Config, Category = Callbacks, Meta = (ClampMin = 0, ForceUnits = "cm/s", EditCondition = "bSubmersionCallbackEnabled"))
-	float MinVelocityForSurfaceTouchCallback = 10.f;
+	float MinVelocityForSurfaceTouchCallback = 100.f;
 
 	virtual void PostInitProperties() override;
 
