Address CIS issue in UE non-unity

diff --git a/Engine/Plugins/Experimental/Buoyancy/Source/Runtime/Public/BuoyancyEventInterface.h b/Engine/Plugins/Experimental/Buoyancy/Source/Runtime/Public/BuoyancyEventInterface.h
--- a/Engine/Plugins/Experimental/Buoyancy/Source/Runtime/Public/BuoyancyEventInterface.h
+++ b/Engine/Plugins/Experimental/Buoyancy/Source/Runtime/Public/BuoyancyEventInterface.h
@@ -2,6 +2,9 @@
 
 #pragma once
 #include "CoreMinimal.h"
+#include "UObject/ObjectMacros.h"
+#include "UObject/ScriptMacros.h"
+#include "UObject/Interface.h"
 #include "BuoyancyEventInterface.generated.h"
 
 UINTERFACE(Blueprintable)
