Break up buoyancy surface touch callback into begin/continue/end events, and make it so you don't need to subscribe to all of them.

diff --git a/Engine/Plugins/Experimental/Buoyancy/Source/Runtime/Private/BuoyancySubsystem.cpp b/Engine/Plugins/Experimental/Buoyancy/Source/Runtime/Private/BuoyancySubsystem.cpp
--- a/Engine/Plugins/Experimental/Buoyancy/Source/Runtime/Private/BuoyancySubsystem.cpp
+++ b/Engine/Plugins/Experimental/Buoyancy/Source/Runtime/Private/BuoyancySubsystem.cpp
@@ -124,13 +124,17 @@ void UBuoyancySubsystem::ApplyRuntimeSettings(const UBuoyancyRuntimeSettings* In
 	BuoyancySettings.MinBoundsSubdivisionVol = InSettings->MinBoundsSubdivisionVol;
 	BuoyancySettings.MinVelocityForSurfaceTouchCallback = InSettings->MinVelocityForSurfaceTouchCallback;
 
-	UWorld* World = GetWorld();
-	BuoyancySettings.bSurfaceTouchCallback =
+	// Based on server/client/editor, determine if we should generate callbacks.
+	// If we're editor, always generate callbacks. If we're not editor, only
+	// generate callbacks on client.	
 #if WITH_EDITOR
-		InSettings->bSurfaceTouchCallbackOnServer || InSettings->bSurfaceTouchCallbackOnClient;
+	BuoyancySettings.SurfaceTouchCallbackFlags = InSettings->SurfaceTouchCallbackFlags;
 #else
-		(World && World->IsNetMode(NM_DedicatedServer) && InSettings->bSurfaceTouchCallbackOnServer) ||
-		(World && World->IsNetMode(NM_Client) && InSettings->bSurfaceTouchCallbackOnClient);
+	UWorld* World = GetWorld();
+	BuoyancySettings.SurfaceTouchCallbackFlags
+		= (World && World->IsNetMode(NM_Client))
+		? InSettings->SurfaceTouchCallbackFlags
+		: EBuoyancyEventFlags::None;
 #endif
 
 	// Enable or disable
@@ -178,7 +182,7 @@ void UBuoyancySubsystem::Tick(float DeltaTime)
 	}
 
 	// Process surface-touched callbacks
-	if (BuoyancySettings.bSurfaceTouchCallback)
+	if (BuoyancySettings.SurfaceTouchCallbackFlags != 0)
 	{
 		SCOPE_CYCLE_COUNTER(STAT_BuoyancySubsystem_DispatchCallbacks)
 
@@ -186,21 +190,18 @@ void UBuoyancySubsystem::Tick(float DeltaTime)
 		{
 			for (const FBuoyancySubsystemSimCallbackOutput::FSurfaceTouch& SurfaceTouch : AsyncOutput->SurfaceTouches)
 			{
+				// Skip if we mask out this touch type
+				if ((SurfaceTouch.Flag & BuoyancySettings.SurfaceTouchCallbackFlags) == 0)
+				{
+					continue;
+				}
+
 				// Extract primitive components
 				UPrimitiveComponent* WaterComponent = PhysScene->GetOwningComponent<UPrimitiveComponent>(SurfaceTouch.WaterProxy);
 				UPrimitiveComponent* RigidComponent = PhysScene->GetOwningComponent<UPrimitiveComponent>(SurfaceTouch.RigidProxy);
 
 				// Get the parental water body component
-				AWaterBody* WaterBodyActor = WaterComponent->GetOwner<AWaterBody>();
-
-				// Dispatch the delegate
-				OnSurfaceTouched.Broadcast(
-					WaterBodyActor,
-					WaterComponent,
-					RigidComponent,
-					SurfaceTouch.Vol,
-					SurfaceTouch.CoM,
-					SurfaceTouch.Vel);
+				AWaterBody* WaterActor = WaterComponent->GetOwner<AWaterBody>();
 
 				const auto DispatchEvent = [&](AActor* Actor)
 				{
@@ -209,18 +210,30 @@ void UBuoyancySubsystem::Tick(float DeltaTime)
 					// If the actor implements the event interface, call the surface touched callback
 					if (Actor->Implements<UBuoyancyEventInterface>())
 					{
-						IBuoyancyEventInterface::Execute_OnSurfaceTouched(
-							Actor,
-							WaterBodyActor,
-							WaterComponent,
-							RigidComponent,
-							SurfaceTouch.Vol,
-							SurfaceTouch.CoM,
-							SurfaceTouch.Vel);
+						switch (SurfaceTouch.Flag)
+						{
+							case EBuoyancyEventFlags::Begin:
+								IBuoyancyEventInterface::Execute_OnSurfaceTouchBegin(
+									Actor, WaterActor, WaterComponent, RigidComponent,
+									SurfaceTouch.Vol, SurfaceTouch.CoM, SurfaceTouch.Vel);
+								break;
+
+							case EBuoyancyEventFlags::Continue:
+								IBuoyancyEventInterface::Execute_OnSurfaceTouching(
+									Actor, WaterActor, WaterComponent, RigidComponent,
+									SurfaceTouch.Vol, SurfaceTouch.CoM, SurfaceTouch.Vel);
+								break;
+
+							case EBuoyancyEventFlags::End:
+								IBuoyancyEventInterface::Execute_OnSurfaceTouchEnd(
+									Actor, WaterActor, WaterComponent, RigidComponent,
+									SurfaceTouch.Vol, SurfaceTouch.CoM, SurfaceTouch.Vel);
+								break;
+						}
 					}
 				};
-				DispatchEvent(WaterBodyActor);
-				DispatchEvent(WaterComponent->GetOwner());
+				DispatchEvent(WaterActor);
+				DispatchEvent(RigidComponent->GetOwner());
 			}
 		}
 	}
@@ -338,6 +351,12 @@ void FBuoyancySubsystemSimCallback::OnMidPhaseModification_Internal(Chaos::FMidP
 	// How much time has the sim ticked this frame
 	const FReal DeltaSeconds = GetDeltaTime_Internal();
 
+	// SparseArray implements move semantics, so these swaps should amount to pointer swaps.
+	// This way array memories stick around even when reset/swapped so we don't do many
+	// new allocations.
+	Swap(Submersions, PrevSubmersions);
+	Swap(SubmersionMetaData, PrevSubmersionMetaData);
+
 	// Clear submersions and submerged shapes array, but keep their memory allocated
 	Submersions.Reset();
 	SubmergedShapes.Reset();
@@ -420,7 +439,8 @@ void FBuoyancySubsystemSimCallback::OnMidPhaseModification_Internal(Chaos::FMidP
 					}
 
 					// If this is a surface touch record it for callback, if 
-					if (BuoyancySettings->bSurfaceTouchCallback && TotalVol > SubmergedVol * (1.f + UE_KINDA_SMALL_NUMBER))
+					if (BuoyancySettings->SurfaceTouchCallbackFlags != 0 &&
+						TotalVol > SubmergedVol * (1.f + UE_KINDA_SMALL_NUMBER))
 					{
 						SCOPE_CYCLE_COUNTER(STAT_BuoyancySubsystem_BuildSubmersionCallbackData)
 
@@ -498,31 +518,52 @@ void FBuoyancySubsystemSimCallback::OnMidPhaseModification_Internal(Chaos::FMidP
 
 
 	// Generate callback data if we're into that sort of thing
-	if (BuoyancySettings->bSurfaceTouchCallback)
+	const uint8 CallbackFlags = BuoyancySettings->SurfaceTouchCallbackFlags;
+	if (CallbackFlags != 0)
 	{
 		SCOPE_CYCLE_COUNTER(STAT_BuoyancySubsystem_ProduceSurfaceTouches)
 
 		// Get the async output struct to write to
 		FBuoyancySubsystemSimCallbackOutput& Output = GetProducerOutputData_Internal();
 
-		// Make sure we have metadata for each submersion
-		//ensureAlwaysMsgf(Submersions.Num() == SubmersionMetaData.Num());
-
 		// Process every surface touch and queue up some of them to return
 		// to game thread for callback dispatch
 		Output.SurfaceTouches.Reserve(SubmersionMetaData.Num() * FBuoyancySubmersionMetaData::MaxNumWaterContacts);
 		for (auto Iter = SubmersionMetaData.CreateIterator(); Iter; ++Iter)
 		{
-			//
-			// TODO: Check to see if this is touching surface
-			//
-
 			const FBuoyancySubmersionMetaData& MetaData = *Iter;
-			const FBuoyancySubmersion& Submersion = Submersions[Iter.GetIndex()];
+			const int32 ObjectIndex = Iter.GetIndex();
+			const FBuoyancySubmersion& Submersion = Submersions[ObjectIndex];
+
+			// Mark this as a new or continuing contact based on whether or
+			// not we have a bit from the previous-submersions array.
+			const bool bPrevSubmerged =
+				PrevSubmersionMetaData.IsValidIndex(ObjectIndex) &&
+				PrevSubmersionMetaData.IsAllocated(ObjectIndex);
+			const EBuoyancyEventFlags TouchFlag
+				= bPrevSubmerged
+				? EBuoyancyEventFlags::Continue
+				: EBuoyancyEventFlags::Begin;
+
+			// Clear out the "prev" entry for this one's metadata so that
+			// we can loop over the prev metadata for lost-contacts. Only
+			// bother doing this work if we're tracking removals
+			if (bPrevSubmerged && (CallbackFlags & EBuoyancyEventFlags::End) != 0)
+			{
+				PrevSubmersionMetaData.RemoveAt(ObjectIndex);
+			}
+
+			// Only continue if we're tracking this touch type
+			if ((TouchFlag & CallbackFlags) == 0)
+			{
+				continue;
+			}
 
+			// Build up output of new and continuing surface touches
 			for (const FBuoyancySubmersionMetaData::FWaterContact& WaterContact : MetaData.WaterContacts)
 			{
 				Output.SurfaceTouches.Add({
+					TouchFlag,
 					Submersion.Particle->PhysicsProxy(),
 					WaterContact.Water->PhysicsProxy(),
 					WaterContact.Vol,
@@ -531,5 +572,31 @@ void FBuoyancySubsystemSimCallback::OnMidPhaseModification_Internal(Chaos::FMidP
 				});
 			}
 		}
+
+		// The remaining previous submersion metadata will correspond with lost contacts
+		//
+		// NOTE:
+		// At the moment, lost contact callbacks will only occur when an entire object
+		// loses contact, not just when one part of it loses contact.
+		if ((CallbackFlags & EBuoyancyEventFlags::End) != 0)
+		{
+			for (auto Iter = PrevSubmersionMetaData.CreateIterator(); Iter; ++Iter)
+			{
+				const FBuoyancySubmersionMetaData& MetaData = *Iter;
+				const int32 ObjectIndex = Iter.GetIndex();
+				const FBuoyancySubmersion& Submersion = PrevSubmersions[ObjectIndex];
+				for (const FBuoyancySubmersionMetaData::FWaterContact& WaterContact : MetaData.WaterContacts)
+				{
+					Output.SurfaceTouches.Add({
+						EBuoyancyEventFlags::End,
+						Submersion.Particle->PhysicsProxy(),
+						WaterContact.Water->PhysicsProxy(),
+						WaterContact.Vol,
+						WaterContact.CoM,
+						WaterContact.Vel
+					});
+				}
+			}
+		}
 	}
 }
diff --git a/Engine/Plugins/Experimental/Buoyancy/Source/Runtime/Public/BuoyancyEventFlags.h b/Engine/Plugins/Experimental/Buoyancy/Source/Runtime/Public/BuoyancyEventFlags.h
new file mode 100644
--- /dev/null
+++ b/Engine/Plugins/Experimental/Buoyancy/Source/Runtime/Public/BuoyancyEventFlags.h
@@ -0,0 +1,16 @@
+// Copyright Epic Games, Inc. All Rights Reserved.
+
+#pragma once
+#include "CoreMinimal.h"
+#include "Misc/EnumClassFlags.h"
+#include "BuoyancyEventFlags.generated.h"
+
+UENUM(BlueprintType, meta=(Bitflags, UseEnumValuesAsMaskValuesInEditor="true"))
+enum EBuoyancyEventFlags : uint8
+{
+	None		= 0 UMETA(Hidden),
+	Begin		= 1 << 0,
+	Continue	= 1 << 1,
+	End			= 1 << 2,
+};
+ENUM_CLASS_FLAGS(EBuoyancyEventFlags);
diff --git a/Engine/Plugins/Experimental/Buoyancy/Source/Runtime/Public/BuoyancyEventInterface.h b/Engine/Plugins/Experimental/Buoyancy/Source/Runtime/Public/BuoyancyEventInterface.h
--- a/Engine/Plugins/Experimental/Buoyancy/Source/Runtime/Public/BuoyancyEventInterface.h
+++ b/Engine/Plugins/Experimental/Buoyancy/Source/Runtime/Public/BuoyancyEventInterface.h
@@ -19,8 +19,13 @@ class BUOYANCY_API IBuoyancyEventInterface
 
 public:
 
+	/**
+	 * Called when part of a submerged simulated actor first comes in
+	 * contact with a water surface. This can be called for multiple parts
+	 * of a complex body.
+	 */
 	UFUNCTION(BlueprintImplementableEvent, Category = WaterBody)
-	void  OnSurfaceTouched(
+	void  OnSurfaceTouchBegin(
 		class AWaterBody* WaterBodyActor,
 		UPrimitiveComponent* WaterComponent,
 		UPrimitiveComponent* SubmergedComponent,
@@ -28,4 +33,29 @@ public:
 		const FVector& SubmergedCenterOfMass,
 		const FVector& SubmergedVelocity);
 
+	/**
+	 * Called continually while objects maintain contact with a water surface.
+	 * May be called multiple times for different parts of an object.
+	 */
+	UFUNCTION(BlueprintImplementableEvent, Category = WaterBody)
+	void  OnSurfaceTouching(
+		class AWaterBody* WaterBodyActor,
+		UPrimitiveComponent* WaterComponent,
+		UPrimitiveComponent* SubmergedComponent,
+		float SubmergedVolume,
+		const FVector& SubmergedCenterOfMass,
+		const FVector& SubmergedVelocity);
+
+	/**
+	 * Called when a submerged body loses contact with all water surfaces. This
+	 * can result from total submersion or from coming completely out of water.
+	 */
+	UFUNCTION(BlueprintImplementableEvent, Category = WaterBody)
+	void  OnSurfaceTouchEnd(
+		class AWaterBody* WaterBodyActor,
+		UPrimitiveComponent* WaterComponent,
+		UPrimitiveComponent* SubmergedComponent,
+		float SubmergedVolume,
+		const FVector& SubmergedCenterOfMass,
+		const FVector& SubmergedVelocity);
 };
diff --git a/Engine/Plugins/Experimental/Buoyancy/Source/Runtime/Public/BuoyancyRuntimeSettings.h b/Engine/Plugins/Experimental/Buoyancy/Source/Runtime/Public/BuoyancyRuntimeSettings.h
--- a/Engine/Plugins/Experimental/Buoyancy/Source/Runtime/Public/BuoyancyRuntimeSettings.h
+++ b/Engine/Plugins/Experimental/Buoyancy/Source/Runtime/Public/BuoyancyRuntimeSettings.h
@@ -4,6 +4,7 @@
 
 #include "Engine/DeveloperSettings.h"
 #include "Templates/SubclassOf.h"
+#include "BuoyancyEventFlags.h"
 #include "BuoyancyRuntimeSettings.generated.h"
 
 enum ECollisionChannel : int;
@@ -29,35 +30,29 @@ public:
 	bool bKeepFloatingObjectsAwake = false;
 
 	/** Density of water to use in buoyancy calculations. The density of water is approximately 1g/cm3 */
-	UPROPERTY(EditAnywhere, Config, Category = Buoyancy, Meta = (ClampMin = 0, ForceUnits = "g/cm3"))
+	UPROPERTY(EditAnywhere, Config, Category = WaterProperties, Meta = (ClampMin = 0, ForceUnits = "g/cm3"))
 	float WaterDensity = 1.f;
 
 	/** Drag factor for submerged objects. This unitless number approximates the idea of viscosity,
 	    but internally functions identically to the "ether drag" concept. */
-	UPROPERTY(EditAnywhere, Config, Category = Buoyancy, Meta = (ClampMin = 0))
+	UPROPERTY(EditAnywhere, Config, Category = WaterProperties, Meta = (ClampMin = 0))
 	float WaterDrag = 1.f;
 
+	/** Collision channel to use for water ObjectTypes */
+	UPROPERTY(EditAnywhere, Config, Category = WaterProperties)
+	TEnumAsByte<ECollisionChannel> CollisionChannelForWaterObjects;
+
 	/** Maximum number of times that a buoyancy bounds can be split into 8 */
-	UPROPERTY(EditAnywhere, Config, Category = Buoyancy, Meta = (ClampMin = 0))
+	UPROPERTY(EditAnywhere, Config, Category = SubmergedVolumeCalculation, Meta = (ClampMin = 0))
 	int32 MaxNumBoundsSubdivisions = 2;
 
 	/** Minimum volume bounding box which can be produced by a subdivision operation */
-	UPROPERTY(EditAnywhere, Config, Category = Buoyancy)
+	UPROPERTY(EditAnywhere, Config, Category = SubmergedVolumeCalculation)
 	float MinBoundsSubdivisionVol = FMath::Pow(100.f, 3.f); // 1m^3
 
-	/** Collision channel to use for water ObjectTypes */
-	UPROPERTY(EditAnywhere, Config, Category = Collision)
-	TEnumAsByte<ECollisionChannel> CollisionChannelForWaterObjects;
-
-	/** When enabled, a delegate will be enabled which triggers when objects
-	    information for generating splashes and waves. */
-	UPROPERTY(EditAnywhere, Config, Category = Callbacks)
-	bool bSurfaceTouchCallbackOnServer = false;
-
-	/** When enabled, a delegate will be enabled which triggers when objects
-	    are touching the surface of a body of water on client. */
-	UPROPERTY(EditAnywhere, Config, Category = Callbacks)
-	bool bSurfaceTouchCallbackOnClient = true;
+	/** Callback data for water surface contacts will be generated according to these flags */
+	UPROPERTY(EditAnywhere, Config, Category = Callbacks, Meta = (Bitmask, BitmaskEnum = "/Script/Buoyancy.EBuoyancyEventFlags"))
+	uint8 SurfaceTouchCallbackFlags = EBuoyancyEventFlags::Begin | EBuoyancyEventFlags::End;
 
 	/** Minimum velocity of an object relative to a water required in order to trigger
 	    a submersion callback. If this is zero, then the callback will be triggered
diff --git a/Engine/Plugins/Experimental/Buoyancy/Source/Runtime/Public/BuoyancySubsystem.h b/Engine/Plugins/Experimental/Buoyancy/Source/Runtime/Public/BuoyancySubsystem.h
--- a/Engine/Plugins/Experimental/Buoyancy/Source/Runtime/Public/BuoyancySubsystem.h
+++ b/Engine/Plugins/Experimental/Buoyancy/Source/Runtime/Public/BuoyancySubsystem.h
@@ -8,6 +8,7 @@
 #include "PBDRigidsSolver.h"
 #include "Engine/EngineBaseTypes.h"
 #include "WaterBodyComponent.h"
+#include "BuoyancyEventFlags.h"
 #include "BuoyancySubsystem.generated.h"
 
 DECLARE_LOG_CATEGORY_EXTERN(LogBuoyancySubsystem, Log, All);
@@ -37,26 +38,12 @@ struct FBuoyancySettings
 
 	ECollisionChannel WaterCollisionChannel = ECollisionChannel::ECC_MAX;
 
-	bool bSurfaceTouchCallback = true;
+	uint8 SurfaceTouchCallbackFlags = EBuoyancyEventFlags::None;
 
 	float MinVelocityForSurfaceTouchCallback = 10.f;
 };
 
 
-//
-// Buoyancy Delegate Callback
-//
-
-UDELEGATE()
-DECLARE_DYNAMIC_MULTICAST_DELEGATE_SixParams(FSurfaceTouchedDelegate,
-	class AWaterBody*, WaterBodyActor,
-	UPrimitiveComponent*, WaterComponent,
-	UPrimitiveComponent*, SubmergedComponent,
-	float, SubmergedVolume,
-	const FVector&, SubmergedCenterOfMass,
-	const FVector&, SubmergedVelocity);
-
-
 //
 // Buoyancy Subsystem
 //
@@ -83,11 +70,6 @@ public:
 	UFUNCTION()
 	bool IsEnabled() const;
 
-	// Give access to a delegate which will trigger when objects
-	// are in water.
-	UPROPERTY(BlueprintAssignable, Category = Buoyancy)
-	FSurfaceTouchedDelegate OnSurfaceTouched;
-
 protected:
 
 	// UTickableWorldSubsystem begin interface
@@ -159,6 +141,7 @@ struct FBuoyancySubsystemSimCallbackOutput : public Chaos::FSimCallbackOutput
 {
 	struct FSurfaceTouch
 	{
+		uint8 Flag;
 		IPhysicsProxyBase* RigidProxy;
 		IPhysicsProxyBase* WaterProxy;
 		float Vol;
@@ -193,10 +176,12 @@ private:
 	// a member variable and reset every frame, to avoid reallocation of similarly
 	// sized data.
 	TSparseArray<FBuoyancySubmersion> Submersions;
+	TSparseArray<FBuoyancySubmersion> PrevSubmersions;
 
 	// Another sparse array to be kept in sync with Submersions, which will contain
 	// metadata useful for event callbacks
 	TSparseArray<FBuoyancySubmersionMetaData> SubmersionMetaData;
+	TSparseArray<FBuoyancySubmersionMetaData> PrevSubmersionMetaData;
 
 	// This is a sparse array of bit arrays representing which shapes in an object
 	// have already been accounted for when submerging an object. For example, if
