Sync water spline data with physics thread with UserDataPT and use for more accurate water queries.
Adjust drag based on relative velocity to water to enable currents.

diff --git a/Engine/Plugins/Experimental/Buoyancy/Buoyancy.uplugin b/Engine/Plugins/Experimental/Buoyancy/Buoyancy.uplugin
--- a/Engine/Plugins/Experimental/Buoyancy/Buoyancy.uplugin
+++ b/Engine/Plugins/Experimental/Buoyancy/Buoyancy.uplugin
@@ -26,5 +26,9 @@
 			"Name": "Water",
 			"Enabled": true
 		},
+		{
+			"Name": "ChaosUserDataPT",
+			"Enabled": true
+		}
 	]
 }
\ No newline at end of file
diff --git a/Engine/Plugins/Experimental/Buoyancy/Source/Runtime/Buoyancy.Build.cs b/Engine/Plugins/Experimental/Buoyancy/Source/Runtime/Buoyancy.Build.cs
--- a/Engine/Plugins/Experimental/Buoyancy/Source/Runtime/Buoyancy.Build.cs
+++ b/Engine/Plugins/Experimental/Buoyancy/Source/Runtime/Buoyancy.Build.cs
@@ -14,6 +14,7 @@ public class Buoyancy : ModuleRules
 				"CoreUObject",
 				"Engine",
 				"Water",
+				"ChaosUserDataPT",
 			}
 		);
 
diff --git a/Engine/Plugins/Experimental/Buoyancy/Source/Runtime/Private/BuoyancyAlgorithms.cpp b/Engine/Plugins/Experimental/Buoyancy/Source/Runtime/Private/BuoyancyAlgorithms.cpp
--- a/Engine/Plugins/Experimental/Buoyancy/Source/Runtime/Private/BuoyancyAlgorithms.cpp
+++ b/Engine/Plugins/Experimental/Buoyancy/Source/Runtime/Private/BuoyancyAlgorithms.cpp
@@ -111,7 +111,7 @@ namespace BuoyancyAlgorithms
 {
 	using namespace Chaos;
 
-	FRealSingle ComputeParticleVolume(const FPBDRigidsEvolutionGBF* Evolution, const FGeometryParticleHandle* Particle)
+	FRealSingle ComputeParticleVolume(const FPBDRigidsEvolutionGBF& Evolution, const FGeometryParticleHandle* Particle)
 	{
 		const FPBDRigidParticleHandle* Rigid = Particle->CastToRigidParticle();
 		if (Rigid == nullptr)
@@ -119,7 +119,7 @@ namespace BuoyancyAlgorithms
 			return -1.f;
 		}
 
-		const FChaosPhysicsMaterial* ParticleMaterial = Evolution->GetFirstClusteredPhysicsMaterial(Particle);
+		const FChaosPhysicsMaterial* ParticleMaterial = Evolution.GetFirstClusteredPhysicsMaterial(Particle);
 		if (ParticleMaterial == nullptr)
 		{
 			return -1.f;
@@ -204,37 +204,41 @@ namespace BuoyancyAlgorithms
 		return ShapeVol;
 	}
 
-	bool ComputeSubmergedVolume(const FPBDRigidsEvolutionGBF* Evolution, const FGeometryParticleHandle* ParticleA, const FGeometryParticleHandle* ParticleB, int32 NumSubdivisions, float MinVolume, TSparseArray<TBitArray<>>& SubmergedShapes, float& SubmergedVol, FVec3& SubmergedCoM, float& TotalVol)
+	void ScaleSubmergedVolume(const FPBDRigidsEvolutionGBF& Evolution, const FGeometryParticleHandle* Particle, FRealSingle& SubmergedVol, FRealSingle& TotalVol)
 	{
-		if (ComputeSubmergedVolume(ParticleA, ParticleB, NumSubdivisions, MinVolume, SubmergedShapes, SubmergedVol, SubmergedCoM))
+		// Get submerged object's "particle" volume and "shape" volume.
+		//
+		// The particle volume is the theoretical volume of the particle,
+		// derived from its mass and density.
+		//
+		// The shape volume is the volume of all shape bounds which can
+		// possibly count as submerged volumes.
+		const FRealSingle ParticleVol = ComputeParticleVolume(Evolution, Particle);
+		const FRealSingle ShapeVol = ComputeShapeVolume(Particle);
+		TotalVol = ParticleVol;
+
+		// If the submerged vol somehow exceeded the max shape vol, clamp it
+		if (SubmergedVol - ShapeVol > UE_SMALL_NUMBER)
 		{
-			// Get submerged object's "particle" volume and "shape" volume.
-			//
-			// The particle volume is the theoretical volume of the particle,
-			// derived from its mass and density.
-			//
-			// The shape volume is the volume of all shape bounds which can
-			// possibly count as submerged volumes.
-			const FRealSingle ParticleVolB = ComputeParticleVolume(Evolution, ParticleB);
-			const FRealSingle ShapeVolB = ComputeShapeVolume(ParticleB);
-			TotalVol = ParticleVolB;
-
-			// If the submerged vol somehow exceeded the max shape vol, clamp it
-			if (SubmergedVol - ShapeVolB > UE_SMALL_NUMBER)
-			{
-				SubmergedVol = ShapeVolB;
-			}
+			SubmergedVol = ShapeVol;
+		}
 
-			// Adjust the output volume based on the ratio of the material volume and the shape volume.
-			// We expect the shape volume to have overestimated the submerged volume for most shapes,
-			// especially those which are hollow.
-			if (ParticleVolB > UE_SMALL_NUMBER &&
-				ParticleVolB < ShapeVolB)
-			{
-				const float VolRatio = ParticleVolB / ShapeVolB;
-				SubmergedVol *= VolRatio;
-			}
+		// Adjust the output volume based on the ratio of the material volume and the shape volume.
+		// We expect the shape volume to have overestimated the submerged volume for most shapes,
+		// especially those which are hollow.
+		if (ParticleVol > UE_SMALL_NUMBER &&
+			ParticleVol < ShapeVol)
+		{
+			const float VolRatio = ParticleVol / ShapeVol;
+			SubmergedVol *= VolRatio;
+		}
+	}
 
+	bool ComputeSubmergedVolume(const FPBDRigidsEvolutionGBF& Evolution, const FGeometryParticleHandle* ParticleA, const FGeometryParticleHandle* ParticleB, int32 NumSubdivisions, float MinVolume, TSparseArray<TBitArray<>>& SubmergedShapes, float& SubmergedVol, FVec3& SubmergedCoM, float& TotalVol)
+	{
+		if (ComputeSubmergedVolume(ParticleA, ParticleB, NumSubdivisions, MinVolume, SubmergedShapes, SubmergedVol, SubmergedCoM))
+		{
+			ScaleSubmergedVolume(Evolution, ParticleB, SubmergedVol, TotalVol);
 			return true;
 		}
 
@@ -434,6 +438,142 @@ namespace BuoyancyAlgorithms
 		}
 	}
 
+	bool ComputeSubmergedVolume(const FPBDRigidsEvolutionGBF& Evolution, const FGeometryParticleHandle* SubmergedParticle, const FGeometryParticleHandle* WaterParticle, const float WaterZ, int32 NumSubdivisions, float MinVolume, TSparseArray<TBitArray<>>& SubmergedShapes, float& SubmergedVol, FVec3& SubmergedCoM, float& TotalVol)
+	{
+		if (ComputeSubmergedVolume(SubmergedParticle, WaterParticle, WaterZ, NumSubdivisions, MinVolume, SubmergedShapes, SubmergedVol, SubmergedCoM))
+		{
+			ScaleSubmergedVolume(Evolution, SubmergedParticle, SubmergedVol, TotalVol);
+
+
+#if ENABLE_DRAW_DEBUG
+			if (bBuoyancyDebugDraw)
+			{
+				Chaos::FDebugDrawQueue::GetInstance().DrawDebugPoint(SubmergedCoM, FColor::Yellow, false, -1.f, -1, 15.f);
+			}
+#endif
+
+
+			return true;
+		}
+
+		return false;
+	}
+
+	bool ComputeSubmergedVolume(const FGeometryParticleHandle* SubmergedParticle, const FGeometryParticleHandle* WaterParticle, const float WaterZ, int32 NumSubdivisions, float MinVolume, TSparseArray<TBitArray<>>& SubmergedShapes, float& SubmergedVol, FVec3& SubmergedCoM)
+	{
+		// Get some initial data about the submerged particle
+		const FImplicitObject* RootImplicit = SubmergedParticle->GetGeometry();
+		const FShapeInstanceArray& ShapeInstances = SubmergedParticle->ShapeInstances();
+		const FConstGenericParticleHandle SubmergedGeneric = SubmergedParticle;
+		const FRigidTransform3 ParticleWorldTransform = SubmergedGeneric->GetTransformPQ();
+		const int32 ParticleIndex = SubmergedParticle->UniqueIdx().Idx;
+
+		// Some info about the water
+		const FImplicitObject* WaterRootImplicit = WaterParticle->GetGeometry();
+		const EImplicitObjectType WaterShapeType = Private::GetImplicitCollisionType(WaterParticle, WaterRootImplicit);
+		const FShapeInstanceArray& WaterShapeInstances = WaterParticle->ShapeInstances();
+		const FShapeInstance* WaterShapeInstance = ShapeInstances[0].Get();
+
+		// Initialize submersion values
+		SubmergedVol = 0.f;
+		SubmergedCoM = FVec3::ZeroVector;
+
+		// Traverse the submerged particle's leaves
+		RootImplicit->VisitLeafObjects(
+			[SubmergedParticle, ParticleIndex, &ShapeInstances, WaterShapeType, WaterShapeInstance, &ParticleWorldTransform, WaterZ, &NumSubdivisions, &MinVolume, &SubmergedShapes, &SubmergedVol, &SubmergedCoM]
+			(const FImplicitObject* Implicit, const FRigidTransform3& RelativeTransform, const int32 RootObjectIndex, const int32 ObjectIndex, const int32 LeafObjectIndex)
+		{
+			const FAABB3 RelativeBounds = Implicit->CalculateTransformedBounds(RelativeTransform);
+			const int32 ShapeIndex = (ShapeInstances.IsValidIndex(RootObjectIndex)) ? RootObjectIndex : 0;
+			const FShapeInstance* ShapeInstance = ShapeInstances[ShapeIndex].Get();
+			const EImplicitObjectType ShapeType = Private::GetImplicitCollisionType(SubmergedParticle, Implicit);
+
+			// If this shape pair doesn't pass a narrow phase test then skip it
+			if (!ShapePairNarrowPhaseFilter(ShapeType, ShapeInstance, WaterShapeType, WaterShapeInstance))
+			{
+				return;
+			}
+
+			// If this shape has already been submerged, skip it to avoid double-counting
+			// any buoyancy contributions.
+			if (IsShapeSubmerged_Internal(SubmergedShapes, ParticleIndex, ObjectIndex))
+			{
+				return;
+			}
+
+			// Get the world-space bounds of shape A
+			const FRigidTransform3 ShapeWorldTransform = RelativeTransform * ParticleWorldTransform;
+			const FAABB3 LocalBox = Implicit->BoundingBox();
+			const FAABB3 WorldBox = LocalBox.TransformedAABB(ShapeWorldTransform);
+
+#if ENABLE_DRAW_DEBUG
+			if (bBuoyancyDebugDraw)
+			{
+				Chaos::FDebugDrawQueue::GetInstance().DrawDebugBox(
+					ShapeWorldTransform.TransformPosition(LocalBox.GetCenter()),
+					LocalBox.Extents() * .5f,
+					ShapeWorldTransform.GetRotation(),
+					FColor::Green, false, -1.f, SDPG_Foreground, 1.f);
+			}
+#endif
+
+
+			// Generate subdivided bounds list
+			TArray<FAABB3> SubmergedBoxes;
+			SubdivideBounds(LocalBox, NumSubdivisions, MinVolume, SubmergedBoxes);
+
+			// Loop over every subdivision of the shape bounds, counting up submerged portions
+			bool bSubmerged = false;
+			for (const FAABB3& Box : SubmergedBoxes)
+			{
+				// Compute the portion of the object bounds that are submerged
+				FAABB3 SubmergedBox;
+				if (ComputeSubmergedBounds(WaterZ, Box, ShapeWorldTransform, SubmergedBox))
+				{
+					// At this point we know that the shape is submerged
+					bSubmerged = true;
+
+					// This bounds box is submerged. Compute it's volume and center of mass
+					// in world space, and add those contributions to the submerged quantity.
+					const FVec3 LeafSubmergedCoM = ShapeWorldTransform.TransformPosition(SubmergedBox.GetCenter());
+					const float LeafSubmergedVol = SubmergedBox.GetVolume();
+					SubmergedVol += LeafSubmergedVol;
+					SubmergedCoM += LeafSubmergedCoM * LeafSubmergedVol;
+
+					// Make sure the volume of the submerged portion never exceeds the total
+					// volume of the leaf bounds
+					const float LeafMaxVol = LocalBox.GetVolume() + UE_SMALL_NUMBER;
+					ensureAlwaysMsgf(LeafSubmergedVol <= LeafMaxVol, TEXT("BuoyancyAlgorithms::ComputeSubmergedVolume: The volume of the submerged portion of the leaf bounds has somehow exceeded the volume of the overall leaf bounds."));
+
+
+
+#if ENABLE_DRAW_DEBUG
+					if (bBuoyancyDebugDraw)
+					{
+						Chaos::FDebugDrawQueue::GetInstance().DrawDebugBox(
+							ShapeWorldTransform.TransformPosition(SubmergedBox.GetCenter()),
+							SubmergedBox.Extents() * .5f,
+							ShapeWorldTransform.GetRotation(),
+							FColor::Red, false, -1.f, SDPG_Foreground, 1.f);
+					}
+#endif
+				}
+			}
+
+			if (bSubmerged)
+			{
+				SubmergeShape_Internal(SubmergedShapes, ParticleIndex, ObjectIndex);
+			}
+		});
+
+		if (SubmergedVol > SMALL_NUMBER)
+		{
+			SubmergedCoM /= SubmergedVol;
+			return true;
+		}
+
+		return false;
+	}
 
 	bool ComputeSubmergedBounds(float WaterZ, const FAABB3& RigidBox, const FRigidTransform3& RigidTransform, FAABB3& OutSubmergedBounds)
 	{
@@ -562,7 +702,7 @@ namespace BuoyancyAlgorithms
 		return false;
 	}
 
-	bool ComputeBuoyantForce(const FPBDRigidParticleHandle* RigidParticle, const float DeltaSeconds, const float WaterDensity, const float WaterDrag, const FVec3& GravityAccelVec, const FVec3& SubmergedCoM, const float SubmergedVol, FVec3& OutDeltaV, FVec3& OutDeltaW)
+	bool ComputeBuoyantForce(const FPBDRigidParticleHandle* RigidParticle, const float DeltaSeconds, const float WaterDensity, const float WaterDrag, const FVec3& GravityAccelVec, const FVec3& SubmergedCoM, const float SubmergedVol, const FVec3& WaterVel, FVec3& OutDeltaV, FVec3& OutDeltaW)
 	{
 		SCOPE_CYCLE_COUNTER(STAT_BuoyancyAlgorithms_ComputeBuoyantForces)
 
@@ -607,6 +747,11 @@ namespace BuoyancyAlgorithms
 		OutDeltaV = LinearAccel * DeltaSeconds;
 		OutDeltaW = AngularAccel * DeltaSeconds;
 
+		// Get the velocities of the submerged portion relative to the water - We want the
+		// drag force to bring these values to zero
+		const FVec3 SubmergedV = RigidParticle->V() + FVec3::CrossProduct(RigidParticle->W(), CoMDiff);
+		const FVec3 RelativeV = SubmergedV - WaterVel;
+
 		// Compute water drag force
 		//
 		// NOTE: This is a very approximate "ether drag" style model here, probably 
@@ -616,7 +761,7 @@ namespace BuoyancyAlgorithms
 		const float DragFactor = FMath::Max(0.f, 1.f - (WaterDrag * DeltaSeconds));
 
 		// Account for water drag in deltas
-		OutDeltaV = (DragFactor * OutDeltaV) + (DragFactor - 1.f) * RigidParticle->V();
+		OutDeltaV = (DragFactor * OutDeltaV) + (DragFactor - 1.f) * RelativeV;
 		OutDeltaW = (DragFactor * OutDeltaW) + (DragFactor - 1.f) * RigidParticle->W();
 
 		//
diff --git a/Engine/Plugins/Experimental/Buoyancy/Source/Runtime/Private/BuoyancySubsystem.cpp b/Engine/Plugins/Experimental/Buoyancy/Source/Runtime/Private/BuoyancySubsystem.cpp
--- a/Engine/Plugins/Experimental/Buoyancy/Source/Runtime/Private/BuoyancySubsystem.cpp
+++ b/Engine/Plugins/Experimental/Buoyancy/Source/Runtime/Private/BuoyancySubsystem.cpp
@@ -13,10 +13,17 @@
 #include "DrawDebugHelpers.h"
 #include "Physics/Experimental/PhysScene_Chaos.h"
 #include "WaterBodyActor.h"
+#include "Components/SplineComponent.h"
+#include "WaterSubsystem.h"
+#include "WaterBodyManager.h"
+#include "WaterSplineComponent.h"
+#include "Chaos/PhysicsObject.h"
 #include "PBDRigidsSolver.h"			// Only needed to get perparticlegravity :(
 #include "Chaos/PBDRigidsEvolutionGBF.h"// Only needed to get perparticlegravity :(
 #include "Chaos/PerParticleGravity.h"	// Needed in order to determine force of gravity on each particle
 #include "Chaos/DebugDrawQueue.h"
+#include "Chaos/PhysicsObjectInternalInterface.h"
+#include "Chaos/PhysicsObject.h"
 
 //
 // CVars
@@ -62,20 +69,27 @@ bool UBuoyancySubsystem::IsEnabled() const
 	return SimCallback != nullptr;
 }
 
-bool UBuoyancySubsystem::CreateSimCallback()
+void UBuoyancySubsystem::CreateSimCallback()
 {
 	// Create sim callback
 	if (Chaos::FPhysicsSolver* Solver = GetSolver())
 	{
+		// Create the callback for keeping spline data in sync
+		SplineData = Solver->CreateAndRegisterSimCallbackObject_External<FBuoyancyWaterSplineDataManager>();
+
+		// Create the main buoyancy sim callback
 		SimCallback = Solver->CreateAndRegisterSimCallbackObject_External<FBuoyancySubsystemSimCallback>();
 
-		return SimCallback != nullptr;
+		// Give the buoyancy sim callback a reference to the spline data callback,
+		// so that it'll have access to per-particle spline data
+		if (FBuoyancySubsystemSimCallbackInput* AsyncInput = SimCallback->GetProducerInputData_External())
+		{
+			AsyncInput->SplineData = SplineData;
+		}
 	}
-
-	return false;
 }
 
-bool UBuoyancySubsystem::DestroySimCallback()
+void UBuoyancySubsystem::DestroySimCallback()
 {
 	if (SimCallback)
 	{
@@ -83,11 +97,8 @@ bool UBuoyancySubsystem::DestroySimCallback()
 		{
 			Solver->UnregisterAndFreeSimCallbackObject_External(SimCallback);
 			SimCallback = nullptr;
-			return true;
 		}
 	}
-
-	return false;
 }
 
 void UBuoyancySubsystem::PostInitialize()
@@ -97,6 +108,14 @@ void UBuoyancySubsystem::PostInitialize()
 	// Apply initial runtime settings
 	ApplyRuntimeSettings(GetDefault<UBuoyancyRuntimeSettings>(), EPropertyChangeType::ValueSet);
 
+	// Setup callback for when waterbodies are added/removed
+	if (FWaterBodyManager* WaterBodyManager = UWaterSubsystem::GetWaterBodyManager(GetWorld()))
+	{
+		WaterBodyManager->OnWaterBodyAdded.AddUObject(this, &UBuoyancySubsystem::OnWaterBodyAdded);
+		WaterBodyManager->OnWaterBodyRemoved.AddUObject(this, &UBuoyancySubsystem::OnWaterBodyRemoved);
+		bWaterObjectsChanged = true;
+	}
+
 	// Set up callback for when runtime settings change in editor
 #if WITH_EDITOR
 	GetDefault<UBuoyancyRuntimeSettings>()->OnSettingsChange.AddUObject(this, &UBuoyancySubsystem::ApplyRuntimeSettings);
@@ -159,6 +178,11 @@ void UBuoyancySubsystem::Tick(float DeltaTime)
 		return;
 	}
 
+	if (SplineData == nullptr)
+	{
+		return;
+	}
+
 	UWorld* World = GetWorld();
 	if (World == nullptr)
 	{
@@ -171,6 +195,50 @@ void UBuoyancySubsystem::Tick(float DeltaTime)
 		return;
 	}
 
+	// Update spline info for all water bodies & internal arrays of water objects
+	if (bWaterObjectsChanged)
+	{
+		if (FWaterBodyManager* WaterBodyManager = UWaterSubsystem::GetWaterBodyManager(GetWorld()))
+		{
+			if (FBuoyancySubsystemSimCallbackInput* AsyncInput = SimCallback->GetProducerInputData_External())
+			{
+				bWaterObjectsChanged = false;
+				AsyncInput->WaterObjects = TArray<Chaos::FPhysicsObjectHandle>();
+				TArray<Chaos::FPhysicsObjectHandle>& WaterObjects = *AsyncInput->WaterObjects;
+				WaterObjects.Reserve(WaterBodyManager->NumWaterBodies());
+				WaterBodyManager->ForEachWaterBodyComponent(GetWorld(), [this, &WaterObjects](UWaterBodyComponent* WaterBodyComponent)
+				{
+					// Get the metadata object, if there is one
+					UWaterSplineMetadata* WaterSplineMetadata = WaterBodyComponent->GetWaterSplineMetadata();
+
+					if (UWaterSplineComponent* SplineComponent = WaterBodyComponent->GetWaterSpline())
+					{
+						// Copy out water spline data into a shared ptr, to be associated with all
+						// child particles and marshaled to PT.
+						const Chaos::FRigidTransform3 WaterTransform = WaterBodyComponent->GetComponentTransform();
+						TSharedPtr<FBuoyancyWaterSplineData> WaterSplineData = MakeShared<FBuoyancyWaterSplineData>(
+							WaterTransform,
+							SplineComponent->SplineCurves.Position,
+							WaterSplineMetadata ? WaterSplineMetadata->WaterVelocityScalar : TOptional<FInterpCurveFloat>()
+						);
+
+						// Go over each physics object in each primitive component which was generated
+						// from this spline, and associate the spline with the particle.
+						for (UPrimitiveComponent* WaterPrimitiveComponent : WaterBodyComponent->GetCollisionComponents(true))
+						{
+							for (Chaos::FPhysicsObjectHandle WaterObject : WaterPrimitiveComponent->GetAllPhysicsObjects())
+							{
+								SplineData->SetData_GT(WaterObject, WaterSplineData);
+								WaterObjects.Add(WaterObject);
+							}
+						}
+					}
+					return true;
+				});
+			}
+		}
+	}
+
 	// Only bother sending new async inputs if our buoyancy settings actually changed
 	if (bBuoyancySettingsChanged)
 	{
@@ -244,6 +312,16 @@ TStatId UBuoyancySubsystem::GetStatId() const
 	RETURN_QUICK_DECLARE_CYCLE_STAT(UBuoyancySubsystem, STATGROUP_Tickables);
 }
 
+void UBuoyancySubsystem::OnWaterBodyAdded(UWaterBodyComponent* WaterBodyComponent)
+{
+	bWaterObjectsChanged = true;
+}
+
+void UBuoyancySubsystem::OnWaterBodyRemoved(UWaterBodyComponent* WaterBodyComponent)
+{
+	bWaterObjectsChanged = true;
+}
+
 Chaos::FPhysicsSolver* UBuoyancySubsystem::GetSolver() const
 {
 	if (UWorld* World = GetWorld())
@@ -275,24 +353,6 @@ void FBuoyancySubsystemSimCallbackOutput::Reset()
 	SurfaceTouches.Reset();
 }
 
-namespace
-{
-	using namespace Chaos;
-
-	// Check to see if any of this particle's shapes has a particular collision channel
-	const bool ParticleHasCollision(const FGeometryParticleHandle& Particle, const ECollisionChannel CollisionChannel)
-	{
-		const auto& Shapes = Particle.ShapesArray();
-		uint32 Word3 = 0;
-		for (const auto& Shape : Shapes)
-		{
-			const FCollisionFilterData& ShapeQueryData = Shape->GetQueryData();
-			Word3 |= ShapeQueryData.Word3;
-		}
-		return GetCollisionChannel(Word3) == CollisionChannel;
-	};
-}
-
 void FBuoyancySubsystemSimCallback::OnPreSimulate_Internal()
 {
 	using namespace Chaos;
@@ -302,6 +362,16 @@ void FBuoyancySubsystemSimCallback::OnPreSimulate_Internal()
 	// If we were sent new buoyancy settings or data, update our local sim copy
 	if (const FBuoyancySubsystemSimCallbackInput* Input = GetConsumerInput_Internal())
 	{
+		if (Input->WaterObjects.IsSet())
+		{
+			WaterObjects = *Input->WaterObjects;
+		}
+
+		if (Input->SplineData.IsSet())
+		{
+			SplineData = *Input->SplineData;
+		}
+
 		if (Input->BuoyancySettings.IsValid())
 		{
 			BuoyancySettings = MoveTemp(Input->BuoyancySettings);
@@ -321,6 +391,12 @@ void FBuoyancySubsystemSimCallback::OnMidPhaseModification_Internal(Chaos::FMidP
 
 	SCOPE_CYCLE_COUNTER(STAT_BuoyancySubsystem_OnMidPhaseModification)
 
+	// If we don't have a spline data manager, early out
+	if (SplineData == nullptr)
+	{
+		return;
+	}
+
 	// If we don't have a valid buoyancy settings object, don't continue
 	if (BuoyancySettings.IsValid() == false)
 	{
@@ -331,7 +407,7 @@ void FBuoyancySubsystemSimCallback::OnMidPhaseModification_Internal(Chaos::FMidP
 	FPBDRigidsEvolution* Evolution = nullptr;
 	if (FPhysicsSolverBase* SolverBase = GetSolver())
 	{
-		// Why does castchecked return a ref? That makes me think
+		// Why does cast-checked return a ref? That makes me think
 		// it's not actually doing a check...
 		FPBDRigidsSolver& PBDSolver = SolverBase->CastChecked();
 		Evolution = PBDSolver.GetEvolution();
@@ -341,16 +417,6 @@ void FBuoyancySubsystemSimCallback::OnMidPhaseModification_Internal(Chaos::FMidP
 		return;
 	}
 
-	// Get perparticle gravity rule, for figuring out the effective gravity on buoyant objects
-	const FPerParticleGravity* PerParticleGravity = &Evolution->GetGravityForces();
-	if (PerParticleGravity == nullptr)
-	{
-		return;
-	}
-
-	// How much time has the sim ticked this frame
-	const FReal DeltaSeconds = GetDeltaTime_Internal();
-
 	// SparseArray implements move semantics, so these swaps should amount to pointer swaps.
 	// This way array memories stick around even when reset/swapped so we don't do many
 	// new allocations.
@@ -362,161 +428,228 @@ void FBuoyancySubsystemSimCallback::OnMidPhaseModification_Internal(Chaos::FMidP
 	SubmergedShapes.Reset();
 	SubmersionMetaData.Reset();
 
+	// Build list of "submersions"
+	ProcessMidPhases(*Evolution, MidPhaseAccessor);
+
+	// Apply buoyant forces resulting from submersions
+	ApplyBuoyantForces(*Evolution);
+
+	// Generate async outputs for callback data
+	GenerateCallbackData();
+}
+
+void FBuoyancySubsystemSimCallback::ProcessMidPhases(
+	Chaos::FPBDRigidsEvolution& Evolution,
+	Chaos::FMidPhaseModifierAccessor& MidPhaseAccessor)
+{
+	SCOPE_CYCLE_COUNTER(STAT_BuoyancySubsystem_VisitMidphases)
+
+	// Loop over water objects
+	Chaos::FReadPhysicsObjectInterface_Internal Interface = Chaos::FPhysicsObjectInternalInterface::GetRead();
+	for (Chaos::FPhysicsObjectHandle WaterObject : WaterObjects)
+	{
+		// Get particle handle from physics object
+		Chaos::FGeometryParticleHandle* WaterParticle = Interface.GetParticle(WaterObject);
+		if (WaterParticle == nullptr) { continue; }
+
+		// Get midphases for this waterobject
+		for (Chaos::FMidPhaseModifier& MidPhase : MidPhaseAccessor.GetMidPhases(WaterParticle))
+		{
+			ProcessMidPhase(Evolution, WaterParticle, MidPhase);
+		}
+	}
+}
+
+void FBuoyancySubsystemSimCallback::ProcessMidPhase(
+	Chaos::FPBDRigidsEvolution& Evolution,
+	Chaos::FGeometryParticleHandle* WaterParticle,
+	Chaos::FMidPhaseModifier& MidPhase)
+{
+	SCOPE_CYCLE_COUNTER(STAT_BuoyancySubsystem_ProcessMidphase)
+
+	using namespace Chaos;
+
+	// Get midphase particles
+	FGeometryParticleHandle* Particle0;
+	FGeometryParticleHandle* Particle1;
+	MidPhase.GetParticles(&Particle0, &Particle1);
+	if (Particle0 == nullptr || Particle1 == nullptr)
+	{
+		return;
+	}
+
+	// Get water spline data
+	const TSharedPtr<FBuoyancyWaterSplineData>* SplineDataPtr = SplineData->GetData_PT(*WaterParticle);
+	const FBuoyancyWaterSplineData* WaterSpline = (SplineDataPtr ? (*SplineDataPtr).Get() : nullptr);
+	if (WaterSpline == nullptr)
 	{
-		SCOPE_CYCLE_COUNTER(STAT_BuoyancySubsystem_VisitMidphases)
+		return;
+	}
 
-		// NOTE: For now we visit _every_ midphase, and check for ones which involve
-		// our target collision channel. It would be nice if it were possible instead
-		// to get a list of water body particles and loop over only midphases which
-		// involve them.
-		MidPhaseAccessor.VisitMidPhases([this, Evolution](Chaos::FMidPhaseModifier& MidPhase)
+	// Select & cast the rigid particle
+	FPBDRigidParticleHandle* RigidParticle = (Particle0 == WaterParticle ? Particle1 : Particle0)->CastToRigidParticle();
+
+	// Find water surface at the nearest point on the spline
+	const FVector ParticlePos = RigidParticle->XCom();
+	const FVector ParticleLocalPos = WaterSpline->Transform.InverseTransformPosition(ParticlePos);
+	float ParticleDistance;
+	const float ClosestSplineKey = WaterSpline->Position.FindNearest(ParticleLocalPos, ParticleDistance);
+	const FVector ClosestSplinePoint = WaterSpline->Transform.TransformPosition(WaterSpline->Position.Eval(ClosestSplineKey));
+	const float WaterZ = ClosestSplinePoint.Z;
+
+	// Get the water velocity at this point
+	const FVec3 WaterVel
+		= WaterSpline->Velocity.IsSet()
+		? WaterSpline->Velocity->Eval(ClosestSplineKey) * WaterSpline->Position.EvalDerivative(ClosestSplineKey).GetSafeNormal()
+		: FVec3::ZeroVector;
+
+#if ENABLE_DRAW_DEBUG
+	if (bBuoyancyDebugDraw)
+	{
+		// Spline Color
+		const FColor SplineColor = FColor::Cyan;
+
+		// Draw projection onto the line
+		const FVec3 SurfacePoint(ParticlePos.X, ParticlePos.Y, WaterZ);
+		Chaos::FDebugDrawQueue::GetInstance().DrawDebugLine(ParticlePos, SurfacePoint, SplineColor, false, -1.f, -1, 6.f);
+		Chaos::FDebugDrawQueue::GetInstance().DrawDebugLine(SurfacePoint, ClosestSplinePoint, SplineColor, false, -1.f, -1, 3.f);
+
+		// Draw a section of the spline near the spline key
+		FVec3 PrevPoint;
+		bool bFirst = true;
+		for (float SplineKey = ClosestSplineKey - .1f; SplineKey <= ClosestSplineKey + .1f; SplineKey += .05f)
 		{
-			FGeometryParticleHandle* Particle0;
-			FGeometryParticleHandle* Particle1;
-			MidPhase.GetParticles(&Particle0, &Particle1);
-			if (Particle0 == nullptr || Particle1 == nullptr)
+			const FVector SplinePoint = WaterSpline->Transform.TransformPosition(WaterSpline->Position.Eval(SplineKey));
+			if (bFirst)
 			{
-				return;
+				bFirst = false;
+				PrevPoint = SplinePoint;
 			}
-
-			// Use collision filters to pick out which of the particles is water, if either.
-			//
-			// TODO: Instead of this, water particles could have additional userdata, or could
-			// register themselves with the callback and we could look for only the midphases
-			// generated with them.
-			const bool bWater0 = ParticleHasCollision(*Particle0, BuoyancySettings->WaterCollisionChannel);
-			const bool bWater1 = ParticleHasCollision(*Particle1, BuoyancySettings->WaterCollisionChannel);
-			if ((bWater0 || bWater1) == false)
+			else
 			{
-				return;
+				Chaos::FDebugDrawQueue::GetInstance().DrawDebugDirectionalArrow(PrevPoint, SplinePoint, 15.f, SplineColor, false, -1.f, -1, 3.f);
 			}
+		}
 
-			// We want to disable all collisions with water, so disable the midphase now
-			MidPhase.Disable();
+		Chaos::FDebugDrawQueue::GetInstance().DrawDebugDirectionalArrow(SurfacePoint, SurfacePoint + WaterVel, 20.f, FColor::Yellow, false, -1.f, -1, 3.f);
+	}
+#endif
 
-			// Store the rigid particle and the water particle
-			if (bWater0) { Swap(Particle0, Particle1); }
-			FPBDRigidParticleHandle* RigidParticle = Particle0->CastToRigidParticle();
-			FGeometryParticleHandle* WaterParticle = Particle1;
-			if (RigidParticle == nullptr)
-			{
-				return;
-			}
+	// Compute submerged volume and CoM
+	float SubmergedVol;
+	FVec3 SubmergedCoM;
+	float TotalVol;
+	if (BuoyancyAlgorithms::ComputeSubmergedVolume(Evolution, RigidParticle, WaterParticle, WaterZ, BuoyancySettings->MaxNumBoundsSubdivisions, BuoyancySettings->MinBoundsSubdivisionVol, SubmergedShapes, SubmergedVol, SubmergedCoM, TotalVol))
+	{
+		SCOPE_CYCLE_COUNTER(STAT_BuoyancySubsystem_BuildSubmersions)
 
-			// Compute submerged volume and CoM
-			float SubmergedVol;
-			FVec3 SubmergedCoM;
-			float TotalVol;
-			if (BuoyancyAlgorithms::ComputeSubmergedVolume(Evolution, WaterParticle, RigidParticle, BuoyancySettings->MaxNumBoundsSubdivisions, BuoyancySettings->MinBoundsSubdivisionVol, SubmergedShapes, SubmergedVol, SubmergedCoM, TotalVol))
+		// If any volume was submerged, take the weighted average of the centers of mass to get the
+		// approximate submerged center of mass, and apply the buoyancy force there.
+		if (SubmergedVol > SMALL_NUMBER)
+		{
+			const int32 RigidParticleIndex = RigidParticle->UniqueIdx().Idx;
+
+			// If this particle was already marked submerged, add to its existing submersion
+			if (Submersions.IsValidIndex(RigidParticleIndex))
 			{
-				SCOPE_CYCLE_COUNTER(STAT_BuoyancySubsystem_BuildSubmersions)
+				FBuoyancySubmersion& Submersion = Submersions[RigidParticleIndex];
+				ensureMsgf(Submersion.Particle == RigidParticle, TEXT("Something went wrong - there's a particle index mismatch in the Submersions sparse array"));
 
-				// If any volume was submerged, take the weighted average of the centers of mass to get the
-				// approximate submerged center of mass, and apply the buoyancy force there.
-				if (SubmergedVol > SMALL_NUMBER)
-				{
-					const int32 RigidParticleIndex = RigidParticle->UniqueIdx().Idx;
+				// Sum the volumes
+				Submersion.Vol = Submersion.Vol + SubmergedVol;
 
-					// If this particle was already marked submerged, add to its existing submersion
-					if (Submersions.IsValidIndex(RigidParticleIndex))
-					{
-						FBuoyancySubmersion& Submersion = Submersions[RigidParticleIndex];
-						ensureMsgf(Submersion.Particle == RigidParticle, TEXT("Something went wrong - there's a particle index mismatch in the Submersions sparse array"));
+				// Get the weighted-average CoM
+				// NOTE: The unchecked division should be safe since we already
+				// know SubmergedVol > SMALL_NUMBER
+				Submersion.CoM = ((Submersion.CoM * Submersion.Vol) + (SubmergedCoM * SubmergedVol)) / Submersion.Vol;
+			}
 
-						// Sum the volumes
-						Submersion.Vol = Submersion.Vol + SubmergedVol;
+			// If this particle was not yet submerged, make a new submersion for it
+			else
+			{
+				Submersions.Insert(RigidParticleIndex, { RigidParticle, SubmergedVol, SubmergedCoM, WaterVel });
+			}
 
-						// Get the weighted-average CoM
-						// NOTE: The unchecked division should be safe since we already
-						// know SubmergedVol > SMALL_NUMBER
-						Submersion.CoM = ((Submersion.CoM * Submersion.Vol) + (SubmergedCoM * SubmergedVol)) / Submersion.Vol;
-					}
+			// If this is a surface touch record it for callback, if 
+			if (BuoyancySettings->SurfaceTouchCallbackFlags != 0 &&
+				TotalVol > SubmergedVol * (1.f + UE_KINDA_SMALL_NUMBER))
+			{
+				SCOPE_CYCLE_COUNTER(STAT_BuoyancySubsystem_BuildSubmersionCallbackData)
 
-					// If this particle was not yet submerged, make a new submersion for it
-					else
+				// Proceed only if this CoM is moving fast enough to generate events
+				const FBuoyancySubmersion& Submersion = Submersions[RigidParticleIndex];
+				const FVec3 CoMDiff = Submersion.CoM - RigidParticle->XCom();
+				const FVec3 CoMVel = RigidParticle->V() + FVec3::CrossProduct(RigidParticle->W(), CoMDiff);
+				const float CoMVelSq = FVec3::DotProduct(CoMVel, CoMVel);
+				const float MinVel = BuoyancySettings->MinVelocityForSurfaceTouchCallback;
+				const float MinVelSq = MinVel * MinVel;
+
+				if (CoMVelSq > MinVelSq)
+				{
+					// If we don't have a metadata for this particle yet, add one
+					if (!SubmersionMetaData.IsValidIndex(RigidParticleIndex))
 					{
-						Submersions.Insert(RigidParticleIndex, { RigidParticle, SubmergedVol, SubmergedCoM });
+						SubmersionMetaData.Insert(RigidParticleIndex, FBuoyancySubmersionMetaData());
 					}
+					FBuoyancySubmersionMetaData& MetaData = SubmersionMetaData[RigidParticleIndex];
 
-					// If this is a surface touch record it for callback, if 
-					if (BuoyancySettings->SurfaceTouchCallbackFlags != 0 &&
-						TotalVol > SubmergedVol * (1.f + UE_KINDA_SMALL_NUMBER))
+					// If we haven't already maxed out on water contacts, add one
+					if (MetaData.WaterContacts.Num() < FBuoyancySubmersionMetaData::MaxNumWaterContacts)
 					{
-						SCOPE_CYCLE_COUNTER(STAT_BuoyancySubsystem_BuildSubmersionCallbackData)
-
-						// Proceed only if this CoM is moving fast enough to generate events
-						const FBuoyancySubmersion& Submersion = Submersions[RigidParticleIndex];
-						const FVec3 CoMDiff = Submersion.CoM - RigidParticle->XCom();
-						const FVec3 CoMVel = RigidParticle->V() + FVec3::CrossProduct(RigidParticle->W(), CoMDiff);
-						const float CoMVelSq = FVec3::DotProduct(CoMVel, CoMVel);
-						const float MinVel = BuoyancySettings->MinVelocityForSurfaceTouchCallback;
-						const float MinVelSq = MinVel * MinVel;
-
-#if ENABLE_DRAW_DEBUG
-						if (bBuoyancyDebugDraw)
-						{
-							Chaos::FDebugDrawQueue::GetInstance().DrawDebugDirectionalArrow(
-								Submersion.CoM, Submersion.CoM + CoMVel, 30.f, FColor::Yellow, false, -1, -1, 2.f);
-						}
-#endif
-
-						if (CoMVelSq > MinVelSq)
-						{
-							// If we don't have a metadata for this particle yet, add one
-							if (!SubmersionMetaData.IsValidIndex(RigidParticleIndex))
-							{
-								SubmersionMetaData.Insert(RigidParticleIndex, FBuoyancySubmersionMetaData());
-							}
-							FBuoyancySubmersionMetaData& MetaData = SubmersionMetaData[RigidParticleIndex];
-
-							// If we haven't already maxed out on water contacts, add one
-							if (MetaData.WaterContacts.Num() < FBuoyancySubmersionMetaData::MaxNumWaterContacts)
-							{
-								MetaData.WaterContacts.Add({ WaterParticle, SubmergedVol, SubmergedCoM, CoMVel });
-							}
-						}
+						MetaData.WaterContacts.Add({ WaterParticle, SubmergedVol, SubmergedCoM, CoMVel });
 					}
 				}
 			}
-		});
+		}
 	}
+}
 
-	{
-		SCOPE_CYCLE_COUNTER(STAT_BuoyancySubsystem_ApplyBuoyantForces)
+void FBuoyancySubsystemSimCallback::ApplyBuoyantForces(Chaos::FPBDRigidsEvolution& Evolution)
+{
+	SCOPE_CYCLE_COUNTER(STAT_BuoyancySubsystem_ApplyBuoyantForces)
+
+	using namespace Chaos;
+
+	// How much time has the sim ticked this frame
+	const FReal DeltaSeconds = GetDeltaTime_Internal();
+
+	// Get perparticle gravity rule, for figuring out the effective gravity on buoyant objects
+	const FPerParticleGravity* PerParticleGravity = &Evolution.GetGravityForces();
 
-		// Apply all buoyant forces
-		for (const FBuoyancySubmersion& Submersion : Submersions)
+	// Apply all buoyant forces
+	for (const FBuoyancySubmersion& Submersion : Submersions)
+	{
+		// Figure out the gravity level of the particle
+		const int32 GravityGroupIndex = Submersion.Particle->GravityGroupIndex();
+		const FVec3 GravityAccel
+			= PerParticleGravity != nullptr && GravityGroupIndex != INDEX_NONE
+			? (FVec3)PerParticleGravity->GetAcceleration(GravityGroupIndex)
+			: FVec3::DownVector * 980.f; // Default to "regular" gravity
+
+		// Compute delta linear and angular velocities due to buoyancy. If they're big enough to
+		// matter, apply them
+		FVec3 DeltaV, DeltaW;
+		if (BuoyancyAlgorithms::ComputeBuoyantForce(Submersion.Particle, DeltaSeconds, BuoyancySettings->WaterDensity, BuoyancySettings->WaterDrag, GravityAccel, Submersion.CoM, Submersion.Vol, Submersion.Vel, DeltaV, DeltaW))
 		{
-			// Figure out the gravity level of the particle
-			const int32 GravityGroupIndex = Submersion.Particle->GravityGroupIndex();
-			const FVec3 GravityAccel
-				= PerParticleGravity != nullptr && GravityGroupIndex != INDEX_NONE
-				? (FVec3)PerParticleGravity->GetAcceleration(GravityGroupIndex)
-				: FVec3::DownVector * 980.f; // Default to "regular" gravity
-
-			// Compute delta linear and angular velocities due to buoyancy. If they're big enough to
-			// matter, apply them
-			FVec3 DeltaV, DeltaW;
-			if (BuoyancyAlgorithms::ComputeBuoyantForce(Submersion.Particle, DeltaSeconds, BuoyancySettings->WaterDensity, BuoyancySettings->WaterDrag, GravityAccel, Submersion.CoM, Submersion.Vol, DeltaV, DeltaW))
-			{
-				// Clamp delta velocities
-				DeltaV = DeltaV.GetClampedToSize(0.f, BuoyancySettings->MaxDeltaV);
-				DeltaW = DeltaW.GetClampedToSize(0.f, BuoyancySettings->MaxDeltaW);
+			// Clamp delta velocities
+			DeltaV = DeltaV.GetClampedToSize(0.f, BuoyancySettings->MaxDeltaV);
+			DeltaW = DeltaW.GetClampedToSize(0.f, BuoyancySettings->MaxDeltaW);
 
-				// Apply the deltas
-				Submersion.Particle->SetV(Submersion.Particle->V() + DeltaV);
-				Submersion.Particle->SetW(Submersion.Particle->W() + DeltaW);
+			// Apply the deltas
+			Submersion.Particle->SetV(Submersion.Particle->V() + DeltaV);
+			Submersion.Particle->SetW(Submersion.Particle->W() + DeltaW);
 
-				// Wake up the body??
-				if (Evolution && BuoyancySettings->bKeepAwake)
-				{
-					Evolution->SetParticleObjectState(Submersion.Particle, EObjectStateType::Dynamic);
-				}
+			// Wake up the body??
+			if (BuoyancySettings->bKeepAwake)
+			{
+				Evolution.SetParticleObjectState(Submersion.Particle, EObjectStateType::Dynamic);
 			}
 		}
 	}
+}
 
-
+void FBuoyancySubsystemSimCallback::GenerateCallbackData()
+{
 	// Generate callback data if we're into that sort of thing
 	const uint8 CallbackFlags = BuoyancySettings->SurfaceTouchCallbackFlags;
 	if (CallbackFlags != 0)
diff --git a/Engine/Plugins/Experimental/Buoyancy/Source/Runtime/Public/BuoyancyAlgorithms.h b/Engine/Plugins/Experimental/Buoyancy/Source/Runtime/Public/BuoyancyAlgorithms.h
--- a/Engine/Plugins/Experimental/Buoyancy/Source/Runtime/Public/BuoyancyAlgorithms.h
+++ b/Engine/Plugins/Experimental/Buoyancy/Source/Runtime/Public/BuoyancyAlgorithms.h
@@ -31,19 +31,27 @@ namespace BuoyancyAlgorithms
 
 	// Compute the effective volume of an entire particle based on its material
 	// density and mass.
-	FRealSingle ComputeParticleVolume(const FPBDRigidsEvolutionGBF* Evolution, const FGeometryParticleHandle* Particle);
+	FRealSingle ComputeParticleVolume(const FPBDRigidsEvolutionGBF& Evolution, const FGeometryParticleHandle* Particle);
 
 	// Compute the effective volume of a shape. This method must reflect the
 	// maximum possible output value of the non-scaled ComputeSubmergedVolume.
 	FRealSingle ComputeShapeVolume(const FGeometryParticleHandle* Particle);
 
+	//
+	void ScaleSubmergedVolume(const FPBDRigidsEvolutionGBF& Evolution, const FGeometryParticleHandle* Particle, FRealSingle& SubmergedVol, FRealSingle& TotalVol);
+
 	// Compute an approximate volume and center of mass of particle B submerged in particle A,
 	// adjusting for the volume of the object based on the material density and mass of the object
-	bool ComputeSubmergedVolume(const FPBDRigidsEvolutionGBF* Evolution, const FGeometryParticleHandle* ParticleA, const FGeometryParticleHandle* ParticleB, int32 NumSubdivisions, float MinVolume, TSparseArray<TBitArray<>>& SubmergedShapes, float& SubmergedVol, FVec3& SubmergedCoM, float& TotalVol);
+	bool ComputeSubmergedVolume(const FPBDRigidsEvolutionGBF& Evolution, const FGeometryParticleHandle* ParticleA, const FGeometryParticleHandle* ParticleB, int32 NumSubdivisions, float MinVolume, TSparseArray<TBitArray<>>& SubmergedShapes, float& SubmergedVol, FVec3& SubmergedCoM, float& TotalVol);
 
 	// Compute an approximate volume and center of mass of particle B submerged in particle A
 	bool ComputeSubmergedVolume(const FGeometryParticleHandle* ParticleA, const FGeometryParticleHandle* ParticleB, int32 NumSubdivisions, float MinVolume, TSparseArray<TBitArray<>>& SubmergedShapes, float& SubmergedVol, FVec3& SubmergedCoM);
 
+	// Compute submerged volume given a single waterlevel
+	bool ComputeSubmergedVolume(const FPBDRigidsEvolutionGBF& Evolution, const FGeometryParticleHandle* SubmergedParticle, const FGeometryParticleHandle* WaterParticle, const float WaterZ, int32 NumSubdivisions, float MinVolume, TSparseArray<TBitArray<>>& SubmergedShapes, float& SubmergedVol, FVec3& SubmergedCoM, float& TotalVol);
+
+	bool ComputeSubmergedVolume(const FGeometryParticleHandle* SubmergedParticle, const FGeometryParticleHandle* WaterParticle, const float WaterZ, int32 NumSubdivisions, float MinVolume, TSparseArray<TBitArray<>>& SubmergedShapes, float& SubmergedVol, FVec3& SubmergedCoM);
+
 	// Given an OOBB and a water level, generate another OOBB which is 1. entirely contained
 	// within the input OOBB and 2. entirely contains the portion of the OOBB which is submerged
 	// below the water level.
@@ -55,5 +63,5 @@ namespace BuoyancyAlgorithms
 
 	// Given a rigid particle and it's submerged CoM and Volume, compute delta velocities for
 	// integrated buoyancy forces on an object
-	bool ComputeBuoyantForce(const FPBDRigidParticleHandle* RigidParticle, const float DeltaSeconds, const float WaterDensity, const float WaterDrag, const FVec3& GravityAccelVec, const FVec3& SubmergedCoM, const float SubmergedVol, FVec3& OutDeltaV, FVec3& OutDeltaW);
+	bool ComputeBuoyantForce(const FPBDRigidParticleHandle* RigidParticle, const float DeltaSeconds, const float WaterDensity, const float WaterDrag, const FVec3& GravityAccelVec, const FVec3& SubmergedCoM, const float SubmergedVol, const FVec3& WaterVel, FVec3& OutDeltaV, FVec3& OutDeltaW);
 }
diff --git a/Engine/Plugins/Experimental/Buoyancy/Source/Runtime/Public/BuoyancyStats.h b/Engine/Plugins/Experimental/Buoyancy/Source/Runtime/Public/BuoyancyStats.h
--- a/Engine/Plugins/Experimental/Buoyancy/Source/Runtime/Public/BuoyancyStats.h
+++ b/Engine/Plugins/Experimental/Buoyancy/Source/Runtime/Public/BuoyancyStats.h
@@ -10,6 +10,7 @@ DECLARE_STATS_GROUP(TEXT("Buoyancy"), STATGROUP_Buoyancy, STATCAT_Advanced);
 DECLARE_CYCLE_STAT(TEXT("Buoyancy_Subsystem_OnPreSimulate"), STAT_BuoyancySubsystem_OnPreSimulate, STATGROUP_Buoyancy);
 DECLARE_CYCLE_STAT(TEXT("Buoyancy_Subsystem_OnMidPhaseModification"), STAT_BuoyancySubsystem_OnMidPhaseModification, STATGROUP_Buoyancy);
 DECLARE_CYCLE_STAT(TEXT("Buoyancy_Subsystem_VisitMidphases"), STAT_BuoyancySubsystem_VisitMidphases, STATGROUP_Buoyancy);
+DECLARE_CYCLE_STAT(TEXT("Buoyancy_Subsystem_ProcessMidphase"), STAT_BuoyancySubsystem_ProcessMidphase, STATGROUP_Buoyancy);
 DECLARE_CYCLE_STAT(TEXT("Buoyancy_Subsystem_BuildSubmersions"), STAT_BuoyancySubsystem_BuildSubmersions, STATGROUP_Buoyancy);
 DECLARE_CYCLE_STAT(TEXT("Buoyancy_Subsystem_BuildSubmersionCallbackData"), STAT_BuoyancySubsystem_BuildSubmersionCallbackData, STATGROUP_Buoyancy);
 DECLARE_CYCLE_STAT(TEXT("Buoyancy_Subsystem_ApplyBuoyantForces"), STAT_BuoyancySubsystem_ApplyBuoyantForces, STATGROUP_Buoyancy);
diff --git a/Engine/Plugins/Experimental/Buoyancy/Source/Runtime/Public/BuoyancySubsystem.h b/Engine/Plugins/Experimental/Buoyancy/Source/Runtime/Public/BuoyancySubsystem.h
--- a/Engine/Plugins/Experimental/Buoyancy/Source/Runtime/Public/BuoyancySubsystem.h
+++ b/Engine/Plugins/Experimental/Buoyancy/Source/Runtime/Public/BuoyancySubsystem.h
@@ -5,14 +5,42 @@
 #include "Subsystems/WorldSubsystem.h"
 #include "Chaos/SimCallbackObject.h"
 #include "Chaos/Framework/PhysicsProxyBase.h"
+#include "Chaos/PBDRigidsEvolutionFwd.h"
 #include "PBDRigidsSolver.h"
 #include "Engine/EngineBaseTypes.h"
 #include "WaterBodyComponent.h"
 #include "BuoyancyEventFlags.h"
+#include "ChaosUserDataPT.h"
 #include "BuoyancySubsystem.generated.h"
 
 DECLARE_LOG_CATEGORY_EXTERN(LogBuoyancySubsystem, Log, All);
 
+//
+// Callback object for keeping water splines up to date on the physics thread
+//
+// NOTE: We use shared ptr here because a single water spline might have many
+// particles associated with it, but we'd like to only store a single copy
+// of the spline.
+//
+
+struct FBuoyancyWaterSplineData
+{
+	Chaos::FRigidTransform3 Transform;
+	FInterpCurveVector Position;
+
+	// There might not be any velocity parameter, so keep this one optional
+	TOptional<FInterpCurveFloat> Velocity;
+};
+
+class FBuoyancyWaterSplineDataManager : public Chaos::TUserDataManagerPT< TSharedPtr<struct FBuoyancyWaterSplineData> > { };
+
+namespace Chaos
+{
+	class FMidPhaseModifierAccessor;
+	class FMidPhaseModifier;
+}
+
+
 //
 // Buoyancy Settings
 //
@@ -55,8 +83,10 @@ class BUOYANCY_API UBuoyancySubsystem : public UTickableWorldSubsystem
 
 		UBuoyancySubsystem()
 		: UTickableWorldSubsystem()
+		, bWaterObjectsChanged(false)
 		, bBuoyancySettingsChanged(false)
 		, BuoyancySettings(FBuoyancySettings())
+		, SplineData(nullptr)
 		, SimCallback(nullptr)
 	{ }
 
@@ -79,20 +109,29 @@ protected:
 	virtual TStatId GetStatId() const override;
 	// UTickableWorldSubsystem end interface
 
+	// FWaterBodyManager delegate begin callbacks
+	void OnWaterBodyAdded(UWaterBodyComponent* WaterBodyComponent);
+	void OnWaterBodyRemoved(UWaterBodyComponent* WaterBodyComponent);
+	// FWaterBodyManager delegate end callbacks
+
 private:
 
-	bool CreateSimCallback();
-	bool DestroySimCallback();
+	void CreateSimCallback();
+	void DestroySimCallback();
 
 	Chaos::FPhysicsSolver* GetSolver() const;
 
 	// When water plugin settings change, this callback will apply changes
 	void ApplyRuntimeSettings(const class UBuoyancyRuntimeSettings* InSettings, EPropertyChangeType::Type ChangeType);
 
+	bool bWaterObjectsChanged;
+
 	bool bBuoyancySettingsChanged;
 
 	FBuoyancySettings BuoyancySettings;
 
+	FBuoyancyWaterSplineDataManager* SplineData;
+
 	class FBuoyancySubsystemSimCallback* SimCallback;
 };
 
@@ -124,10 +163,20 @@ struct FBuoyancySubmersion
 	Chaos::FPBDRigidParticleHandle* Particle;
 	float Vol;
 	FVector CoM;
+	FVector Vel;
 };
 
 struct FBuoyancySubsystemSimCallbackInput : public Chaos::FSimCallbackInput
 {
+	// If this array is set, then we need to update our internal list of
+	// water body physics objects to this one. This should occur very
+	// infrequently, if more than once.
+	TOptional<TArray<Chaos::FPhysicsObjectHandle>> WaterObjects;
+
+	// If this ptr is set, then we have a new spline data manager...
+	// That should only probably happen one time
+	TOptional<FBuoyancyWaterSplineDataManager*> SplineData = nullptr;
+
 	// Here we use a unique ptr so that it is possible to provide an async
 	// input _without_ buoyancy settings (which may be eventually desirable
 	// when we eventually are passing lists of water bodies or water wave
@@ -166,6 +215,18 @@ private:
 	virtual void OnPreSimulate_Internal() override;
 	virtual void OnMidPhaseModification_Internal(Chaos::FMidPhaseModifierAccessor& Modifier) override;
 
+	void ProcessMidPhases(Chaos::FPBDRigidsEvolution& Evolution, Chaos::FMidPhaseModifierAccessor& MidPhaseAccessor);
+	void ProcessMidPhase(Chaos::FPBDRigidsEvolution& Evolution, Chaos::FGeometryParticleHandle* WaterParticle, Chaos::FMidPhaseModifier& MidPhase);
+	void ApplyBuoyantForces(Chaos::FPBDRigidsEvolution& Evolution);
+	void GenerateCallbackData();
+
+	// Internal array of physics objects that were created by water components
+	TArray<Chaos::FPhysicsObjectHandle> WaterObjects;
+
+	// Reference to UserDataPT sim callback which manages synchronization of
+	// water spline data
+	FBuoyancyWaterSplineDataManager* SplineData;
+
 	// Initially we won't have any settings - they have to get passed down
 	// via async input. I used TUniquePtr to control access to the same
 	// memory that was allocated by GT to minimize copies.
diff --git a/Engine/Plugins/Experimental/Water/Source/Runtime/Private/WaterBodyManager.cpp b/Engine/Plugins/Experimental/Water/Source/Runtime/Private/WaterBodyManager.cpp
--- a/Engine/Plugins/Experimental/Water/Source/Runtime/Private/WaterBodyManager.cpp
+++ b/Engine/Plugins/Experimental/Water/Source/Runtime/Private/WaterBodyManager.cpp
@@ -23,13 +23,16 @@ void FWaterBodyManager::Deinitialize()
 int32 FWaterBodyManager::AddWaterBodyComponent(UWaterBodyComponent* InWaterBodyComponent)
 {
 	RequestWaveDataRebuild();
-	return WaterBodyComponents.Register(InWaterBodyComponent);
+	const int32 WaterIndex = WaterBodyComponents.Register(InWaterBodyComponent);
+	OnWaterBodyAdded.Broadcast(InWaterBodyComponent);
+	return WaterIndex;
 }
 
 void FWaterBodyManager::RemoveWaterBodyComponent(UWaterBodyComponent* InWaterBodyComponent)
 {
 	RequestWaveDataRebuild();
 	WaterBodyComponents.Unregister(InWaterBodyComponent, InWaterBodyComponent->GetWaterBodyIndex());
+	OnWaterBodyRemoved.Broadcast(InWaterBodyComponent);
 }
 
 int32 FWaterBodyManager::AddWaterZone(AWaterZone* InWaterZone)
diff --git a/Engine/Plugins/Experimental/Water/Source/Runtime/Public/WaterBodyManager.h b/Engine/Plugins/Experimental/Water/Source/Runtime/Public/WaterBodyManager.h
--- a/Engine/Plugins/Experimental/Water/Source/Runtime/Public/WaterBodyManager.h
+++ b/Engine/Plugins/Experimental/Water/Source/Runtime/Public/WaterBodyManager.h
@@ -2,12 +2,16 @@
 
 #pragma once
 
+#include "CoreMinimal.h"
 #include "Templates/SharedPointer.h"
 
 class UWaterBodyComponent;
 class AWaterZone;
 class FWaterViewExtension;
 
+DECLARE_MULTICAST_DELEGATE_OneParam(FWaterBodyEvent, UWaterBodyComponent*);
+
+
 class WATER_API FWaterBodyManager
 {
 public:
@@ -53,6 +57,10 @@ public:
 
 	FWaterViewExtension* GetWaterViewExtension() { return WaterViewExtension.Get(); }
 
+	FWaterBodyEvent OnWaterBodyAdded;
+
+	FWaterBodyEvent OnWaterBodyRemoved;
+
 private:
 
 	/**
