Disable midphases between buoyant objects and water bodies so that they don't generate narrow phases

diff --git a/Engine/Plugins/Experimental/Buoyancy/Source/Runtime/Private/BuoyancySubsystem.cpp b/Engine/Plugins/Experimental/Buoyancy/Source/Runtime/Private/BuoyancySubsystem.cpp
--- a/Engine/Plugins/Experimental/Buoyancy/Source/Runtime/Private/BuoyancySubsystem.cpp
+++ b/Engine/Plugins/Experimental/Buoyancy/Source/Runtime/Private/BuoyancySubsystem.cpp
@@ -469,6 +469,9 @@ void FBuoyancySubsystemSimCallback::ProcessMidPhase(
 
 	using namespace Chaos;
 
+	// Always disable midphases with water
+	MidPhase.Disable();
+
 	// Get midphase particles
 	FGeometryParticleHandle* Particle0;
 	FGeometryParticleHandle* Particle1;
