Fix android build issue

diff --git a/Engine/Plugins/Experimental/Buoyancy/Source/Runtime/Private/BuoyancySubsystem.cpp b/Engine/Plugins/Experimental/Buoyancy/Source/Runtime/Private/BuoyancySubsystem.cpp
--- a/Engine/Plugins/Experimental/Buoyancy/Source/Runtime/Private/BuoyancySubsystem.cpp
+++ b/Engine/Plugins/Experimental/Buoyancy/Source/Runtime/Private/BuoyancySubsystem.cpp
@@ -24,6 +24,7 @@
 #include "Chaos/DebugDrawQueue.h"
 #include "Chaos/PhysicsObjectInternalInterface.h"
 #include "Chaos/PhysicsObject.h"
+#include "Templates/SharedPointer.h"
 
 //
 // CVars
diff --git a/Engine/Plugins/Experimental/Buoyancy/Source/Runtime/Public/BuoyancySubsystem.h b/Engine/Plugins/Experimental/Buoyancy/Source/Runtime/Public/BuoyancySubsystem.h
--- a/Engine/Plugins/Experimental/Buoyancy/Source/Runtime/Public/BuoyancySubsystem.h
+++ b/Engine/Plugins/Experimental/Buoyancy/Source/Runtime/Public/BuoyancySubsystem.h
@@ -25,6 +25,16 @@ DECLARE_LOG_CATEGORY_EXTERN(LogBuoyancySubsystem, Log, All);
 
 struct FBuoyancyWaterSplineData
 {
+	FBuoyancyWaterSplineData() { }
+	FBuoyancyWaterSplineData(
+		const Chaos::FRigidTransform3& InTransform,
+		const FInterpCurveVector& InPosition,
+		const TOptional<FInterpCurveFloat>& InVelocity)
+		: Transform(InTransform)
+		, Position(InPosition)
+		, Velocity(InVelocity)
+	{ }
+
 	Chaos::FRigidTransform3 Transform;
 	FInterpCurveVector Position;
 
@@ -32,7 +42,7 @@ struct FBuoyancyWaterSplineData
 	TOptional<FInterpCurveFloat> Velocity;
 };
 
-class FBuoyancyWaterSplineDataManager : public Chaos::TUserDataManagerPT< TSharedPtr<struct FBuoyancyWaterSplineData> > { };
+class FBuoyancyWaterSplineDataManager : public Chaos::TUserDataManagerPT< TSharedPtr<FBuoyancyWaterSplineData> > { };
 
 namespace Chaos
 {
