Remove "using namespace Chaos" from function scopes in buoyancy subsystem

diff --git a/Engine/Plugins/Experimental/Buoyancy/Source/Runtime/Private/BuoyancyAlgorithms.cpp b/Engine/Plugins/Experimental/Buoyancy/Source/Runtime/Private/BuoyancyAlgorithms.cpp
--- a/Engine/Plugins/Experimental/Buoyancy/Source/Runtime/Private/BuoyancyAlgorithms.cpp
+++ b/Engine/Plugins/Experimental/Buoyancy/Source/Runtime/Private/BuoyancyAlgorithms.cpp
@@ -25,7 +25,8 @@ extern bool bBuoyancyDebugDraw;
 
 namespace
 {
-	using namespace Chaos;
+	using Chaos::FVec3;
+	using Chaos::FAABB3;
 
 	// Check to see if an object's shape is marked as already submerged
 	bool IsShapeSubmerged_Internal(const TSparseArray<TBitArray<>>& SubmergedShapes, const int32 ParticleIndex, const int32 ShapeIndex)
@@ -109,7 +110,21 @@ namespace
 
 namespace BuoyancyAlgorithms
 {
-	using namespace Chaos;
+	using Chaos::FVec3;
+	using Chaos::FAABB3;
+	using Chaos::FAABBEdge;
+	using Chaos::FRealSingle;
+	using Chaos::FPBDRigidsEvolutionGBF;
+	using Chaos::FGeometryParticleHandle;
+	using Chaos::FPBDRigidParticleHandle;
+	using Chaos::FImplicitObject;
+	using Chaos::FRigidTransform3;
+	using Chaos::FMatrix33;
+	using Chaos::FChaosPhysicsMaterial;
+	using Chaos::FShapeInstance;
+	using Chaos::FShapeInstanceArray;
+	using Chaos::EImplicitObjectType;
+	using Chaos::FConstGenericParticleHandle;
 
 	FRealSingle ComputeParticleVolume(const FPBDRigidsEvolutionGBF& Evolution, const FGeometryParticleHandle* Particle)
 	{
@@ -175,10 +190,10 @@ namespace BuoyancyAlgorithms
 			(const FImplicitObject* InnerImplicitObject, const FRigidTransform3&, const int32 RootObjectIndex, const int32, const int32)
 		{
 			const int32 ShapeIndex = ShapeInstances.IsValidIndex(RootObjectIndex) ? RootObjectIndex : 0;
-			const EImplicitObjectType ShapeType = Private::GetImplicitCollisionType(Particle, InnerImplicitObject);
+			const EImplicitObjectType ShapeType = Chaos::Private::GetImplicitCollisionType(Particle, InnerImplicitObject);
 			if (DoCollide(ShapeType, ShapeInstances[ShapeIndex].Get()))
 			{
-				Utilities::CastHelper(*InnerImplicitObject, [&ShapeVol](const auto& Geom)
+				Chaos::Utilities::CastHelper(*InnerImplicitObject, [&ShapeVol](const auto& Geom)
 				{
 					ShapeVol += Geom.BoundingBox().GetVolume();
 				});
@@ -280,8 +295,8 @@ namespace BuoyancyAlgorithms
 			const FAABB3 RelativeBoundsA = ImplicitA->CalculateTransformedBounds(RelativeTransformA);
 			const FAABB3 ShapeBoundsAInB = RelativeBoundsA.TransformedAABB(ParticleTransformAToB);
 			const int32 ShapeIndexA = (ShapeInstancesA.IsValidIndex(RootObjectIndexA)) ? RootObjectIndexA : 0;
-			const FShapeInstance* ShapeInstanceA = ShapeInstancesA[ShapeIndexA].Get();
-			const EImplicitObjectType ShapeTypeA = Private::GetImplicitCollisionType(ParticleA, ImplicitA);
+			const Chaos::FShapeInstance* ShapeInstanceA = ShapeInstancesA[ShapeIndexA].Get();
+			const Chaos::EImplicitObjectType ShapeTypeA = Chaos::Private::GetImplicitCollisionType(ParticleA, ImplicitA);
 
 			// Get the world-space bounds of shape A
 			const FRigidTransform3 ShapeWorldTransformA = RelativeTransformA * ParticleWorldTransformA;
@@ -305,7 +320,7 @@ namespace BuoyancyAlgorithms
 				// Get shape instance data for shape B
 				const int32 ShapeIndexB = (ShapeInstancesB.IsValidIndex(RootObjectIndexB)) ? RootObjectIndexB : 0;
 				const FShapeInstance* ShapeInstanceB = ShapeInstancesB[ShapeIndexB].Get();
-				const EImplicitObjectType ShapeTypeB = Private::GetImplicitCollisionType(ParticleB, ImplicitB);
+				const EImplicitObjectType ShapeTypeB = Chaos::Private::GetImplicitCollisionType(ParticleB, ImplicitB);
 
 				// If this shape pair doesn't pass a narrow phase test then skip it
 				if (!ShapePairNarrowPhaseFilter(ShapeTypeA, ShapeInstanceA, ShapeTypeB, ShapeInstanceB))
@@ -470,7 +485,7 @@ namespace BuoyancyAlgorithms
 
 		// Some info about the water
 		const FImplicitObject* WaterRootImplicit = WaterParticle->GetGeometry();
-		const EImplicitObjectType WaterShapeType = Private::GetImplicitCollisionType(WaterParticle, WaterRootImplicit);
+		const EImplicitObjectType WaterShapeType = Chaos::Private::GetImplicitCollisionType(WaterParticle, WaterRootImplicit);
 		const FShapeInstanceArray& WaterShapeInstances = WaterParticle->ShapeInstances();
 		const FShapeInstance* WaterShapeInstance = ShapeInstances[0].Get();
 
@@ -486,7 +501,7 @@ namespace BuoyancyAlgorithms
 			const FAABB3 RelativeBounds = Implicit->CalculateTransformedBounds(RelativeTransform);
 			const int32 ShapeIndex = (ShapeInstances.IsValidIndex(RootObjectIndex)) ? RootObjectIndex : 0;
 			const FShapeInstance* ShapeInstance = ShapeInstances[ShapeIndex].Get();
-			const EImplicitObjectType ShapeType = Private::GetImplicitCollisionType(SubmergedParticle, Implicit);
+			const EImplicitObjectType ShapeType = Chaos::Private::GetImplicitCollisionType(SubmergedParticle, Implicit);
 
 			// If this shape pair doesn't pass a narrow phase test then skip it
 			if (!ShapePairNarrowPhaseFilter(ShapeType, ShapeInstance, WaterShapeType, WaterShapeInstance))
@@ -733,7 +748,7 @@ namespace BuoyancyAlgorithms
 		// Get inverse inertia data to compute world space accelerations
 		const FVec3 WorldCoM = RigidGeneric->PCom();
 		const FVec3 CoMDiff = SubmergedCoM - WorldCoM;
-		const FMatrix33 WorldInvI = Utilities::ComputeWorldSpaceInertia(RigidGeneric->RCom(), RigidGeneric->ConditionedInvI());
+		const FMatrix33 WorldInvI = Chaos::Utilities::ComputeWorldSpaceInertia(RigidGeneric->RCom(), RigidGeneric->ConditionedInvI());
 
 		// Compute world buoyant force and torque
 		const FVec3 WorldForce = -GravityDir * BuoyantForce;
diff --git a/Engine/Plugins/Experimental/Buoyancy/Source/Runtime/Private/BuoyancySubsystem.cpp b/Engine/Plugins/Experimental/Buoyancy/Source/Runtime/Private/BuoyancySubsystem.cpp
--- a/Engine/Plugins/Experimental/Buoyancy/Source/Runtime/Private/BuoyancySubsystem.cpp
+++ b/Engine/Plugins/Experimental/Buoyancy/Source/Runtime/Private/BuoyancySubsystem.cpp
@@ -356,8 +356,6 @@ void FBuoyancySubsystemSimCallbackOutput::Reset()
 
 void FBuoyancySubsystemSimCallback::OnPreSimulate_Internal()
 {
-	using namespace Chaos;
-
 	SCOPE_CYCLE_COUNTER(STAT_BuoyancySubsystem_OnPreSimulate)
 
 	// If we were sent new buoyancy settings or data, update our local sim copy
@@ -388,8 +386,6 @@ void FBuoyancySubsystemSimCallback::OnPreSimulate_Internal()
 
 void FBuoyancySubsystemSimCallback::OnMidPhaseModification_Internal(Chaos::FMidPhaseModifierAccessor& MidPhaseAccessor)
 {
-	using namespace Chaos;
-
 	SCOPE_CYCLE_COUNTER(STAT_BuoyancySubsystem_OnMidPhaseModification)
 
 	// If we don't have a spline data manager, early out
@@ -405,12 +401,12 @@ void FBuoyancySubsystemSimCallback::OnMidPhaseModification_Internal(Chaos::FMidP
 	}
 
 	// Get the evolution
-	FPBDRigidsEvolution* Evolution = nullptr;
-	if (FPhysicsSolverBase* SolverBase = GetSolver())
+	Chaos::FPBDRigidsEvolution* Evolution = nullptr;
+	if (Chaos::FPhysicsSolverBase* SolverBase = GetSolver())
 	{
 		// Why does cast-checked return a ref? That makes me think
 		// it's not actually doing a check...
-		FPBDRigidsSolver& PBDSolver = SolverBase->CastChecked();
+		Chaos::FPBDRigidsSolver& PBDSolver = SolverBase->CastChecked();
 		Evolution = PBDSolver.GetEvolution();
 	}
 	if (Evolution == nullptr)
@@ -468,14 +464,12 @@ void FBuoyancySubsystemSimCallback::ProcessMidPhase(
 {
 	SCOPE_CYCLE_COUNTER(STAT_BuoyancySubsystem_ProcessMidphase)
 
-	using namespace Chaos;
-
 	// Always disable midphases with water
 	MidPhase.Disable();
 
 	// Get midphase particles
-	FGeometryParticleHandle* Particle0;
-	FGeometryParticleHandle* Particle1;
+	Chaos::FGeometryParticleHandle* Particle0;
+	Chaos::FGeometryParticleHandle* Particle1;
 	MidPhase.GetParticles(&Particle0, &Particle1);
 	if (Particle0 == nullptr || Particle1 == nullptr)
 	{
@@ -491,7 +485,7 @@ void FBuoyancySubsystemSimCallback::ProcessMidPhase(
 	}
 
 	// Select & cast the rigid particle
-	FPBDRigidParticleHandle* RigidParticle = (Particle0 == WaterParticle ? Particle1 : Particle0)->CastToRigidParticle();
+	Chaos::FPBDRigidParticleHandle* RigidParticle = (Particle0 == WaterParticle ? Particle1 : Particle0)->CastToRigidParticle();
 
 	// Find water surface at the nearest point on the spline
 	const FVector ParticlePos = RigidParticle->XCom();
@@ -502,10 +496,10 @@ void FBuoyancySubsystemSimCallback::ProcessMidPhase(
 	const float WaterZ = ClosestSplinePoint.Z;
 
 	// Get the water velocity at this point
-	const FVec3 WaterVel
+	const Chaos::FVec3 WaterVel
 		= WaterSpline->Velocity.IsSet()
 		? WaterSpline->Velocity->Eval(ClosestSplineKey) * WaterSpline->Position.EvalDerivative(ClosestSplineKey).GetSafeNormal()
-		: FVec3::ZeroVector;
+		: Chaos::FVec3::ZeroVector;
 
 #if ENABLE_DRAW_DEBUG
 	if (bBuoyancyDebugDraw)
@@ -514,12 +508,12 @@ void FBuoyancySubsystemSimCallback::ProcessMidPhase(
 		const FColor SplineColor = FColor::Cyan;
 
 		// Draw projection onto the line
-		const FVec3 SurfacePoint(ParticlePos.X, ParticlePos.Y, WaterZ);
+		const Chaos::FVec3 SurfacePoint(ParticlePos.X, ParticlePos.Y, WaterZ);
 		Chaos::FDebugDrawQueue::GetInstance().DrawDebugLine(ParticlePos, SurfacePoint, SplineColor, false, -1.f, -1, 6.f);
 		Chaos::FDebugDrawQueue::GetInstance().DrawDebugLine(SurfacePoint, ClosestSplinePoint, SplineColor, false, -1.f, -1, 3.f);
 
 		// Draw a section of the spline near the spline key
-		FVec3 PrevPoint;
+		Chaos::FVec3 PrevPoint;
 		bool bFirst = true;
 		for (float SplineKey = ClosestSplineKey - .1f; SplineKey <= ClosestSplineKey + .1f; SplineKey += .05f)
 		{
@@ -541,7 +535,7 @@ void FBuoyancySubsystemSimCallback::ProcessMidPhase(
 
 	// Compute submerged volume and CoM
 	float SubmergedVol;
-	FVec3 SubmergedCoM;
+	Chaos::FVec3 SubmergedCoM;
 	float TotalVol;
 	if (BuoyancyAlgorithms::ComputeSubmergedVolume(Evolution, RigidParticle, WaterParticle, WaterZ, BuoyancySettings->MaxNumBoundsSubdivisions, BuoyancySettings->MinBoundsSubdivisionVol, SubmergedShapes, SubmergedVol, SubmergedCoM, TotalVol))
 	{
@@ -582,9 +576,9 @@ void FBuoyancySubsystemSimCallback::ProcessMidPhase(
 
 				// Proceed only if this CoM is moving fast enough to generate events
 				const FBuoyancySubmersion& Submersion = Submersions[RigidParticleIndex];
-				const FVec3 CoMDiff = Submersion.CoM - RigidParticle->XCom();
-				const FVec3 CoMVel = RigidParticle->V() + FVec3::CrossProduct(RigidParticle->W(), CoMDiff);
-				const float CoMVelSq = FVec3::DotProduct(CoMVel, CoMVel);
+				const Chaos::FVec3 CoMDiff = Submersion.CoM - RigidParticle->XCom();
+				const Chaos::FVec3 CoMVel = RigidParticle->V() + Chaos::FVec3::CrossProduct(RigidParticle->W(), CoMDiff);
+				const float CoMVelSq = Chaos::FVec3::DotProduct(CoMVel, CoMVel);
 				const float MinVel = BuoyancySettings->MinVelocityForSurfaceTouchCallback;
 				const float MinVelSq = MinVel * MinVel;
 
@@ -612,27 +606,25 @@ void FBuoyancySubsystemSimCallback::ApplyBuoyantForces(Chaos::FPBDRigidsEvolutio
 {
 	SCOPE_CYCLE_COUNTER(STAT_BuoyancySubsystem_ApplyBuoyantForces)
 
-	using namespace Chaos;
-
 	// How much time has the sim ticked this frame
-	const FReal DeltaSeconds = GetDeltaTime_Internal();
+	const Chaos::FReal DeltaSeconds = GetDeltaTime_Internal();
 
 	// Get perparticle gravity rule, for figuring out the effective gravity on buoyant objects
-	const FPerParticleGravity* PerParticleGravity = &Evolution.GetGravityForces();
+	const Chaos::FPerParticleGravity* PerParticleGravity = &Evolution.GetGravityForces();
 
 	// Apply all buoyant forces
 	for (const FBuoyancySubmersion& Submersion : Submersions)
 	{
 		// Figure out the gravity level of the particle
 		const int32 GravityGroupIndex = Submersion.Particle->GravityGroupIndex();
-		const FVec3 GravityAccel
+		const Chaos::FVec3 GravityAccel
 			= PerParticleGravity != nullptr && GravityGroupIndex != INDEX_NONE
-			? (FVec3)PerParticleGravity->GetAcceleration(GravityGroupIndex)
-			: FVec3::DownVector * 980.f; // Default to "regular" gravity
+			? (Chaos::FVec3)PerParticleGravity->GetAcceleration(GravityGroupIndex)
+			: Chaos::FVec3::DownVector * 980.f; // Default to "regular" gravity
 
 		// Compute delta linear and angular velocities due to buoyancy. If they're big enough to
 		// matter, apply them
-		FVec3 DeltaV, DeltaW;
+		Chaos::FVec3 DeltaV, DeltaW;
 		if (BuoyancyAlgorithms::ComputeBuoyantForce(Submersion.Particle, DeltaSeconds, BuoyancySettings->WaterDensity, BuoyancySettings->WaterDrag, GravityAccel, Submersion.CoM, Submersion.Vol, Submersion.Vel, DeltaV, DeltaW))
 		{
 			// Clamp delta velocities
@@ -646,7 +638,7 @@ void FBuoyancySubsystemSimCallback::ApplyBuoyantForces(Chaos::FPBDRigidsEvolutio
 			// Wake up the body??
 			if (BuoyancySettings->bKeepAwake)
 			{
-				Evolution.SetParticleObjectState(Submersion.Particle, EObjectStateType::Dynamic);
+				Evolution.SetParticleObjectState(Submersion.Particle, Chaos::EObjectStateType::Dynamic);
 			}
 		}
 	}
diff --git a/Engine/Plugins/Experimental/Buoyancy/Source/Runtime/Public/BuoyancyAlgorithms.h b/Engine/Plugins/Experimental/Buoyancy/Source/Runtime/Public/BuoyancyAlgorithms.h
--- a/Engine/Plugins/Experimental/Buoyancy/Source/Runtime/Public/BuoyancyAlgorithms.h
+++ b/Engine/Plugins/Experimental/Buoyancy/Source/Runtime/Public/BuoyancyAlgorithms.h
@@ -14,54 +14,52 @@ namespace Chaos
 
 namespace BuoyancyAlgorithms
 {
-	using namespace Chaos;
-
 	// Minimal struct containing essential data about a particular submersion
 	struct FSubmersion
 	{
 		// Indicates the submerged particle
-		FPBDRigidParticleHandle* SubmergedParticle;
+		Chaos::FPBDRigidParticleHandle* SubmergedParticle;
 
 		// Total submerged volume
 		float SubmergedVolume;
 
 		// Effective submerged center of mass
-		FVec3 SubmergedCoM;
+		Chaos::FVec3 SubmergedCoM;
 	};
 
 	// Compute the effective volume of an entire particle based on its material
 	// density and mass.
-	FRealSingle ComputeParticleVolume(const FPBDRigidsEvolutionGBF& Evolution, const FGeometryParticleHandle* Particle);
+	Chaos::FRealSingle ComputeParticleVolume(const Chaos::FPBDRigidsEvolutionGBF& Evolution, const Chaos::FGeometryParticleHandle* Particle);
 
 	// Compute the effective volume of a shape. This method must reflect the
 	// maximum possible output value of the non-scaled ComputeSubmergedVolume.
-	FRealSingle ComputeShapeVolume(const FGeometryParticleHandle* Particle);
+	Chaos::FRealSingle ComputeShapeVolume(const Chaos::FGeometryParticleHandle* Particle);
 
 	//
-	void ScaleSubmergedVolume(const FPBDRigidsEvolutionGBF& Evolution, const FGeometryParticleHandle* Particle, FRealSingle& SubmergedVol, FRealSingle& TotalVol);
+	void ScaleSubmergedVolume(const Chaos::FPBDRigidsEvolutionGBF& Evolution, const Chaos::FGeometryParticleHandle* Particle, Chaos::FRealSingle& SubmergedVol, Chaos::FRealSingle& TotalVol);
 
 	// Compute an approximate volume and center of mass of particle B submerged in particle A,
 	// adjusting for the volume of the object based on the material density and mass of the object
-	bool ComputeSubmergedVolume(const FPBDRigidsEvolutionGBF& Evolution, const FGeometryParticleHandle* ParticleA, const FGeometryParticleHandle* ParticleB, int32 NumSubdivisions, float MinVolume, TSparseArray<TBitArray<>>& SubmergedShapes, float& SubmergedVol, FVec3& SubmergedCoM, float& TotalVol);
+	bool ComputeSubmergedVolume(const Chaos::FPBDRigidsEvolutionGBF& Evolution, const Chaos::FGeometryParticleHandle* ParticleA, const Chaos::FGeometryParticleHandle* ParticleB, int32 NumSubdivisions, float MinVolume, TSparseArray<TBitArray<>>& SubmergedShapes, float& SubmergedVol, Chaos::FVec3& SubmergedCoM, float& TotalVol);
 
 	// Compute an approximate volume and center of mass of particle B submerged in particle A
-	bool ComputeSubmergedVolume(const FGeometryParticleHandle* ParticleA, const FGeometryParticleHandle* ParticleB, int32 NumSubdivisions, float MinVolume, TSparseArray<TBitArray<>>& SubmergedShapes, float& SubmergedVol, FVec3& SubmergedCoM);
+	bool ComputeSubmergedVolume(const Chaos::FGeometryParticleHandle* ParticleA, const Chaos::FGeometryParticleHandle* ParticleB, int32 NumSubdivisions, float MinVolume, TSparseArray<TBitArray<>>& SubmergedShapes, float& SubmergedVol, Chaos::FVec3& SubmergedCoM);
 
 	// Compute submerged volume given a single waterlevel
-	bool ComputeSubmergedVolume(const FPBDRigidsEvolutionGBF& Evolution, const FGeometryParticleHandle* SubmergedParticle, const FGeometryParticleHandle* WaterParticle, const float WaterZ, int32 NumSubdivisions, float MinVolume, TSparseArray<TBitArray<>>& SubmergedShapes, float& SubmergedVol, FVec3& SubmergedCoM, float& TotalVol);
+	bool ComputeSubmergedVolume(const Chaos::FPBDRigidsEvolutionGBF& Evolution, const Chaos::FGeometryParticleHandle* SubmergedParticle, const Chaos::FGeometryParticleHandle* WaterParticle, const float WaterZ, int32 NumSubdivisions, float MinVolume, TSparseArray<TBitArray<>>& SubmergedShapes, float& SubmergedVol, Chaos::FVec3& SubmergedCoM, float& TotalVol);
 
-	bool ComputeSubmergedVolume(const FGeometryParticleHandle* SubmergedParticle, const FGeometryParticleHandle* WaterParticle, const float WaterZ, int32 NumSubdivisions, float MinVolume, TSparseArray<TBitArray<>>& SubmergedShapes, float& SubmergedVol, FVec3& SubmergedCoM);
+	bool ComputeSubmergedVolume(const Chaos::FGeometryParticleHandle* SubmergedParticle, const Chaos::FGeometryParticleHandle* WaterParticle, const float WaterZ, int32 NumSubdivisions, float MinVolume, TSparseArray<TBitArray<>>& SubmergedShapes, float& SubmergedVol, Chaos::FVec3& SubmergedCoM);
 
 	// Given an OOBB and a water level, generate another OOBB which is 1. entirely contained
 	// within the input OOBB and 2. entirely contains the portion of the OOBB which is submerged
 	// below the water level.
-	bool ComputeSubmergedBounds(float WaterZ, const FAABB3& RigidBox, const FRigidTransform3& RigidTransform, FAABB3& OutSubmergedBounds);
+	bool ComputeSubmergedBounds(float WaterZ, const Chaos::FAABB3& RigidBox, const Chaos::FRigidTransform3& RigidTransform, Chaos::FAABB3& OutSubmergedBounds);
 
 	// Given a bounds object, recursively subdivide it in eighths to a fixed maximum depth and
 	// a fixed minimum smallest subdivision volume.
-	bool SubdivideBounds(const FAABB3& Bounds, int32 NumSubdivisions, float MinVolume, TArray<FAABB3>& OutBounds);
+	bool SubdivideBounds(const Chaos::FAABB3& Bounds, int32 NumSubdivisions, float MinVolume, TArray<Chaos::FAABB3>& OutBounds);
 
 	// Given a rigid particle and it's submerged CoM and Volume, compute delta velocities for
 	// integrated buoyancy forces on an object
-	bool ComputeBuoyantForce(const FPBDRigidParticleHandle* RigidParticle, const float DeltaSeconds, const float WaterDensity, const float WaterDrag, const FVec3& GravityAccelVec, const FVec3& SubmergedCoM, const float SubmergedVol, const FVec3& WaterVel, FVec3& OutDeltaV, FVec3& OutDeltaW);
+	bool ComputeBuoyantForce(const Chaos::FPBDRigidParticleHandle* RigidParticle, const float DeltaSeconds, const float WaterDensity, const float WaterDrag, const Chaos::FVec3& GravityAccelVec, const Chaos::FVec3& SubmergedCoM, const float SubmergedVol, const Chaos::FVec3& WaterVel, Chaos::FVec3& OutDeltaV, Chaos::FVec3& OutDeltaW);
 }
