Fix crash in buoyancy subsystem when water splines are created/destroyed.
Remove line in DefaultEngine.ini which disables buoyancy.

Description of updated strategy:
Rather than trying to maintain a physics thread list of water bodies, just check for water spline data objects. This may result in extra iterations, but sparse array lookups are fast (as fast as particle member lookups like ->X()), and it will never see a discarded particle, unless there's a bug deeper in the midphase generation.

diff --git a/Engine/Plugins/Experimental/Buoyancy/Source/Runtime/Private/BuoyancySubsystem.cpp b/Engine/Plugins/Experimental/Buoyancy/Source/Runtime/Private/BuoyancySubsystem.cpp
--- a/Engine/Plugins/Experimental/Buoyancy/Source/Runtime/Private/BuoyancySubsystem.cpp
+++ b/Engine/Plugins/Experimental/Buoyancy/Source/Runtime/Private/BuoyancySubsystem.cpp
@@ -11,20 +11,18 @@
 #include "Chaos/MidPhaseModification.h"
 #include "Chaos/MassProperties.h"
 #include "DrawDebugHelpers.h"
-#include "Physics/Experimental/PhysScene_Chaos.h"
 #include "WaterBodyActor.h"
 #include "Components/SplineComponent.h"
 #include "WaterSubsystem.h"
 #include "WaterBodyManager.h"
 #include "WaterSplineComponent.h"
 #include "Chaos/PhysicsObject.h"
-#include "PBDRigidsSolver.h"			// Only needed to get perparticlegravity :(
-#include "Chaos/PBDRigidsEvolutionGBF.h"// Only needed to get perparticlegravity :(
-#include "Chaos/PerParticleGravity.h"	// Needed in order to determine force of gravity on each particle
+#include "PBDRigidsSolver.h"
+#include "Chaos/PBDRigidsEvolutionGBF.h"
+#include "Chaos/PerParticleGravity.h"
 #include "Chaos/DebugDrawQueue.h"
-#include "Chaos/PhysicsObjectInternalInterface.h"
-#include "Chaos/PhysicsObject.h"
 #include "Templates/SharedPointer.h"
+#include "PhysicsProxy/SingleParticlePhysicsProxy.h"
 
 //
 // CVars
@@ -179,131 +177,153 @@ void UBuoyancySubsystem::Tick(float DeltaTime)
 		return;
 	}
 
-	if (SplineData == nullptr)
+	// Update spline info for all water bodies & internal arrays of water objects
+	if (bWaterObjectsChanged)
 	{
-		return;
+		UpdateSplineData();
 	}
 
-	UWorld* World = GetWorld();
-	if (World == nullptr)
+	// Only bother sending new async inputs if our buoyancy settings actually changed
+	if (bBuoyancySettingsChanged)
 	{
-		return;
+		UpdateBuoyancySettings();
 	}
 
-	FPhysScene* PhysScene = World->GetPhysicsScene();
-	if (PhysScene == nullptr)
+	// Process surface-touched callbacks
+	if (BuoyancySettings.SurfaceTouchCallbackFlags != 0)
+	{
+		ProcessSurfaceTouchCallbacks();
+	}
+}
+
+void UBuoyancySubsystem::UpdateSplineData()
+{
+	SCOPE_CYCLE_COUNTER(STAT_BuoyancySubsystem_UpdateWaterBodiesList)
+
+	if (SplineData == nullptr)
 	{
 		return;
 	}
 
-	// Update spline info for all water bodies & internal arrays of water objects
-	if (bWaterObjectsChanged)
+	if (FWaterBodyManager* WaterBodyManager = UWaterSubsystem::GetWaterBodyManager(GetWorld()))
 	{
-		if (FWaterBodyManager* WaterBodyManager = UWaterSubsystem::GetWaterBodyManager(GetWorld()))
+		bWaterObjectsChanged = false;
+
+		// Loop over every registered water body
+		WaterBodyManager->ForEachWaterBodyComponent(GetWorld(), [this](UWaterBodyComponent* WaterBodyComponent)
 		{
-			if (FBuoyancySubsystemSimCallbackInput* AsyncInput = SimCallback->GetProducerInputData_External())
+			// Get the metadata object, if there is one
+			UWaterSplineMetadata* WaterSplineMetadata = WaterBodyComponent->GetWaterSplineMetadata();
+
+			if (UWaterSplineComponent* SplineComponent = WaterBodyComponent->GetWaterSpline())
 			{
-				bWaterObjectsChanged = false;
-				AsyncInput->WaterObjects = TArray<Chaos::FPhysicsObjectHandle>();
-				TArray<Chaos::FPhysicsObjectHandle>& WaterObjects = *AsyncInput->WaterObjects;
-				WaterObjects.Reserve(WaterBodyManager->NumWaterBodies());
-				WaterBodyManager->ForEachWaterBodyComponent(GetWorld(), [this, &WaterObjects](UWaterBodyComponent* WaterBodyComponent)
+				// Copy out water spline data into a shared ptr, to be associated with all
+				// child particles and marshaled to PT.
+				const Chaos::FRigidTransform3 WaterTransform = WaterBodyComponent->GetComponentTransform();
+				TSharedPtr<FBuoyancyWaterSplineData> WaterSplineData = MakeShared<FBuoyancyWaterSplineData>(
+					WaterTransform,
+					SplineComponent->SplineCurves.Position,
+					WaterSplineMetadata ? WaterSplineMetadata->WaterVelocityScalar : TOptional<FInterpCurveFloat>()
+				);
+
+				// Go over each physics object in each primitive component which was generated
+				// from this spline, and associate the spline with the particle.
+				for (UPrimitiveComponent* WaterPrimitiveComponent : WaterBodyComponent->GetCollisionComponents(true))
 				{
-					// Get the metadata object, if there is one
-					UWaterSplineMetadata* WaterSplineMetadata = WaterBodyComponent->GetWaterSplineMetadata();
-
-					if (UWaterSplineComponent* SplineComponent = WaterBodyComponent->GetWaterSpline())
+					// Add each object (probably just one) to the objects list
+					if (WaterPrimitiveComponent)
 					{
-						// Copy out water spline data into a shared ptr, to be associated with all
-						// child particles and marshaled to PT.
-						const Chaos::FRigidTransform3 WaterTransform = WaterBodyComponent->GetComponentTransform();
-						TSharedPtr<FBuoyancyWaterSplineData> WaterSplineData = MakeShared<FBuoyancyWaterSplineData>(
-							WaterTransform,
-							SplineComponent->SplineCurves.Position,
-							WaterSplineMetadata ? WaterSplineMetadata->WaterVelocityScalar : TOptional<FInterpCurveFloat>()
-						);
-
-						// Go over each physics object in each primitive component which was generated
-						// from this spline, and associate the spline with the particle.
-						for (UPrimitiveComponent* WaterPrimitiveComponent : WaterBodyComponent->GetCollisionComponents(true))
+						FBodyInstance& BodyInstance = WaterPrimitiveComponent->BodyInstance;
+						if (BodyInstance.IsValidBodyInstance())
 						{
-							for (Chaos::FPhysicsObjectHandle WaterObject : WaterPrimitiveComponent->GetAllPhysicsObjects())
+							if (Chaos::FSingleParticlePhysicsProxy* WaterProxy = BodyInstance.ActorHandle)
 							{
-								SplineData->SetData_GT(WaterObject, WaterSplineData);
-								WaterObjects.Add(WaterObject);
+								SplineData->SetData_GT(WaterProxy->GetGameThreadAPI(), WaterSplineData);
 							}
 						}
 					}
-					return true;
-				});
+				}
 			}
-		}
+			return true;
+		});
 	}
+}
 
-	// Only bother sending new async inputs if our buoyancy settings actually changed
-	if (bBuoyancySettingsChanged)
+void UBuoyancySubsystem::UpdateBuoyancySettings()
+{
+	SCOPE_CYCLE_COUNTER(STAT_BuoyancySubsystem_UpdateBuoyancySettings)
+
+	if (FBuoyancySubsystemSimCallbackInput* AsyncInput = SimCallback->GetProducerInputData_External())
 	{
-		if (FBuoyancySubsystemSimCallbackInput* AsyncInput = SimCallback->GetProducerInputData_External())
-		{
-			bBuoyancySettingsChanged = false;
-			AsyncInput->BuoyancySettings = MakeUnique<FBuoyancySettings>(BuoyancySettings);
-		}
+		bBuoyancySettingsChanged = false;
+		AsyncInput->BuoyancySettings = MakeUnique<FBuoyancySettings>(BuoyancySettings);
 	}
+}
 
-	// Process surface-touched callbacks
-	if (BuoyancySettings.SurfaceTouchCallbackFlags != 0)
+void UBuoyancySubsystem::ProcessSurfaceTouchCallbacks()
+{
+	SCOPE_CYCLE_COUNTER(STAT_BuoyancySubsystem_DispatchCallbacks)
+
+	UWorld* World = GetWorld();
+	if (World == nullptr)
 	{
-		SCOPE_CYCLE_COUNTER(STAT_BuoyancySubsystem_DispatchCallbacks)
+		return;
+	}
 
-		while (Chaos::TSimCallbackOutputHandle<FBuoyancySubsystemSimCallbackOutput> AsyncOutput = SimCallback->PopFutureOutputData_External())
+	FPhysScene* PhysScene = World->GetPhysicsScene();
+	if (PhysScene == nullptr)
+	{
+		return;
+	}
+
+	while (Chaos::TSimCallbackOutputHandle<FBuoyancySubsystemSimCallbackOutput> AsyncOutput = SimCallback->PopFutureOutputData_External())
+	{
+		for (const FBuoyancySubsystemSimCallbackOutput::FSurfaceTouch& SurfaceTouch : AsyncOutput->SurfaceTouches)
 		{
-			for (const FBuoyancySubsystemSimCallbackOutput::FSurfaceTouch& SurfaceTouch : AsyncOutput->SurfaceTouches)
+			// Skip if we mask out this touch type
+			if ((SurfaceTouch.Flag & BuoyancySettings.SurfaceTouchCallbackFlags) == 0)
 			{
-				// Skip if we mask out this touch type
-				if ((SurfaceTouch.Flag & BuoyancySettings.SurfaceTouchCallbackFlags) == 0)
-				{
-					continue;
-				}
+				continue;
+			}
 
-				// Extract primitive components
-				UPrimitiveComponent* WaterComponent = PhysScene->GetOwningComponent<UPrimitiveComponent>(SurfaceTouch.WaterProxy);
-				UPrimitiveComponent* RigidComponent = PhysScene->GetOwningComponent<UPrimitiveComponent>(SurfaceTouch.RigidProxy);
+			// Extract primitive components
+			UPrimitiveComponent* WaterComponent = PhysScene->GetOwningComponent<UPrimitiveComponent>(SurfaceTouch.WaterProxy);
+			UPrimitiveComponent* RigidComponent = PhysScene->GetOwningComponent<UPrimitiveComponent>(SurfaceTouch.RigidProxy);
 
-				// Get the parental water body component
-				AWaterBody* WaterActor = WaterComponent->GetOwner<AWaterBody>();
+			// Get the parental water body component
+			AWaterBody* WaterActor = WaterComponent->GetOwner<AWaterBody>();
 
-				const auto DispatchEvent = [&](AActor* Actor)
-				{
-					// TODO: Actor relevancy check?
+			const auto DispatchEvent = [&](AActor* Actor)
+			{
+				// TODO: Actor relevancy check?
 
-					// If the actor implements the event interface, call the surface touched callback
-					if (Actor->Implements<UBuoyancyEventInterface>())
+				// If the actor implements the event interface, call the surface touched callback
+				if (Actor->Implements<UBuoyancyEventInterface>())
+				{
+					switch (SurfaceTouch.Flag)
 					{
-						switch (SurfaceTouch.Flag)
-						{
-							case EBuoyancyEventFlags::Begin:
-								IBuoyancyEventInterface::Execute_OnSurfaceTouchBegin(
-									Actor, WaterActor, WaterComponent, RigidComponent,
-									SurfaceTouch.Vol, SurfaceTouch.CoM, SurfaceTouch.Vel);
-								break;
-
-							case EBuoyancyEventFlags::Continue:
-								IBuoyancyEventInterface::Execute_OnSurfaceTouching(
-									Actor, WaterActor, WaterComponent, RigidComponent,
-									SurfaceTouch.Vol, SurfaceTouch.CoM, SurfaceTouch.Vel);
-								break;
-
-							case EBuoyancyEventFlags::End:
-								IBuoyancyEventInterface::Execute_OnSurfaceTouchEnd(
-									Actor, WaterActor, WaterComponent, RigidComponent,
-									SurfaceTouch.Vol, SurfaceTouch.CoM, SurfaceTouch.Vel);
-								break;
-						}
+						case EBuoyancyEventFlags::Begin:
+							IBuoyancyEventInterface::Execute_OnSurfaceTouchBegin(
+								Actor, WaterActor, WaterComponent, RigidComponent,
+								SurfaceTouch.Vol, SurfaceTouch.CoM, SurfaceTouch.Vel);
+							break;
+
+						case EBuoyancyEventFlags::Continue:
+							IBuoyancyEventInterface::Execute_OnSurfaceTouching(
+								Actor, WaterActor, WaterComponent, RigidComponent,
+								SurfaceTouch.Vol, SurfaceTouch.CoM, SurfaceTouch.Vel);
+							break;
+
+						case EBuoyancyEventFlags::End:
+							IBuoyancyEventInterface::Execute_OnSurfaceTouchEnd(
+								Actor, WaterActor, WaterComponent, RigidComponent,
+								SurfaceTouch.Vol, SurfaceTouch.CoM, SurfaceTouch.Vel);
+							break;
 					}
-				};
-				DispatchEvent(WaterActor);
-				DispatchEvent(RigidComponent->GetOwner());
-			}
+				}
+			};
+			DispatchEvent(WaterActor);
+			DispatchEvent(RigidComponent->GetOwner());
 		}
 	}
 }
@@ -346,6 +366,7 @@ Chaos::FPhysicsSolver* UBuoyancySubsystem::GetSolver() const
 
 void FBuoyancySubsystemSimCallbackInput::Reset()
 {
+	SplineData.Reset();
 	BuoyancySettings.Reset();
 }
 
@@ -361,11 +382,6 @@ void FBuoyancySubsystemSimCallback::OnPreSimulate_Internal()
 	// If we were sent new buoyancy settings or data, update our local sim copy
 	if (const FBuoyancySubsystemSimCallbackInput* Input = GetConsumerInput_Internal())
 	{
-		if (Input->WaterObjects.IsSet())
-		{
-			WaterObjects = *Input->WaterObjects;
-		}
-
 		if (Input->SplineData.IsSet())
 		{
 			SplineData = *Input->SplineData;
@@ -441,25 +457,45 @@ void FBuoyancySubsystemSimCallback::ProcessMidPhases(
 {
 	SCOPE_CYCLE_COUNTER(STAT_BuoyancySubsystem_VisitMidphases)
 
-	// Loop over water objects
-	Chaos::FReadPhysicsObjectInterface_Internal Interface = Chaos::FPhysicsObjectInternalInterface::GetRead();
-	for (Chaos::FPhysicsObjectHandle WaterObject : WaterObjects)
+	// Loop over all midphases
+	MidPhaseAccessor.VisitMidPhases([this, &Evolution](Chaos::FMidPhaseModifier& MidPhase)
 	{
-		// Get particle handle from physics object
-		Chaos::FGeometryParticleHandle* WaterParticle = Interface.GetParticle(WaterObject);
-		if (WaterParticle == nullptr) { continue; }
-
-		// Get midphases for this waterobject
-		for (Chaos::FMidPhaseModifier& MidPhase : MidPhaseAccessor.GetMidPhases(WaterParticle))
+		// Make sure we have two valid particles
+		Chaos::FGeometryParticleHandle* WaterParticle;
+		Chaos::FGeometryParticleHandle* OtherParticle;
+		MidPhase.GetParticles(&WaterParticle, &OtherParticle);
+		if (WaterParticle && OtherParticle)
 		{
-			ProcessMidPhase(Evolution, WaterParticle, MidPhase);
+			// Get spline data for particle 0. If it exists, then it's water.
+			// If it doesn't exist, then try the other particle.
+			const TSharedPtr<FBuoyancyWaterSplineData>* WaterSpline = SplineData->GetData_PT(*WaterParticle);
+			if (WaterSpline == nullptr || !WaterSpline->IsValid())
+			{
+				// Swap the particles and try again
+				Swap(WaterParticle, OtherParticle);
+				WaterSpline = SplineData->GetData_PT(*WaterParticle);
+				if (WaterSpline == nullptr || !WaterSpline->IsValid())
+				{
+					// Neither particle has a water spline data, so give up.
+					// This is not a water interaction
+					return;
+				}
+			}
+
+			// Make sure the non-water particle is backed by a rigid
+			if (Chaos::FPBDRigidParticleHandle* RigidParticle = OtherParticle->CastToRigidParticle())
+			{
+				ProcessMidPhase(Evolution, WaterParticle, RigidParticle, *WaterSpline->Get(), MidPhase);
+			}
 		}
-	}
+	});
 }
 
 void FBuoyancySubsystemSimCallback::ProcessMidPhase(
 	Chaos::FPBDRigidsEvolution& Evolution,
 	Chaos::FGeometryParticleHandle* WaterParticle,
+	Chaos::FPBDRigidParticleHandle* RigidParticle,
+	const FBuoyancyWaterSplineData& WaterSpline,
 	Chaos::FMidPhaseModifier& MidPhase)
 {
 	SCOPE_CYCLE_COUNTER(STAT_BuoyancySubsystem_ProcessMidphase)
@@ -467,71 +503,58 @@ void FBuoyancySubsystemSimCallback::ProcessMidPhase(
 	// Always disable midphases with water
 	MidPhase.Disable();
 
-	// Get midphase particles
-	Chaos::FGeometryParticleHandle* Particle0;
-	Chaos::FGeometryParticleHandle* Particle1;
-	MidPhase.GetParticles(&Particle0, &Particle1);
-	if (Particle0 == nullptr || Particle1 == nullptr)
-	{
-		return;
-	}
-
-	// Get water spline data
-	const TSharedPtr<FBuoyancyWaterSplineData>* SplineDataPtr = SplineData->GetData_PT(*WaterParticle);
-	const FBuoyancyWaterSplineData* WaterSpline = (SplineDataPtr ? (*SplineDataPtr).Get() : nullptr);
-	if (WaterSpline == nullptr)
+	// Evaluate spline at the object CoM to approximate the water depth
+	float WaterZ;
+	Chaos::FVec3 WaterVel;
 	{
-		return;
-	}
-
-	// Select & cast the rigid particle
-	Chaos::FPBDRigidParticleHandle* RigidParticle = (Particle0 == WaterParticle ? Particle1 : Particle0)->CastToRigidParticle();
-
-	// Find water surface at the nearest point on the spline
-	const FVector ParticlePos = RigidParticle->XCom();
-	const FVector ParticleLocalPos = WaterSpline->Transform.InverseTransformPosition(ParticlePos);
-	float ParticleDistance;
-	const float ClosestSplineKey = WaterSpline->Position.FindNearest(ParticleLocalPos, ParticleDistance);
-	const FVector ClosestSplinePoint = WaterSpline->Transform.TransformPosition(WaterSpline->Position.Eval(ClosestSplineKey));
-	const float WaterZ = ClosestSplinePoint.Z;
-
-	// Get the water velocity at this point
-	const Chaos::FVec3 WaterVel
-		= WaterSpline->Velocity.IsSet()
-		? WaterSpline->Velocity->Eval(ClosestSplineKey) * WaterSpline->Position.EvalDerivative(ClosestSplineKey).GetSafeNormal()
-		: Chaos::FVec3::ZeroVector;
+		SCOPE_CYCLE_COUNTER(STAT_Buoyancy_Subsystem_SplineEvaluation)
+
+		// Find water surface at the nearest point on the spline
+		const FVector ParticlePos = RigidParticle->XCom();
+		const FVector ParticleLocalPos = WaterSpline.Transform.InverseTransformPosition(ParticlePos);
+		float ParticleDistance;
+		const float ClosestSplineKey = WaterSpline.Position.FindNearest(ParticleLocalPos, ParticleDistance);
+		const FVector ClosestSplinePoint = WaterSpline.Transform.TransformPosition(WaterSpline.Position.Eval(ClosestSplineKey));
+		WaterZ = ClosestSplinePoint.Z;
+
+		// Get the water velocity at this point
+		WaterVel
+			= WaterSpline.Velocity.IsSet()
+			? WaterSpline.Velocity->Eval(ClosestSplineKey) * WaterSpline.Position.EvalDerivative(ClosestSplineKey).GetSafeNormal()
+			: Chaos::FVec3::ZeroVector;
 
 #if ENABLE_DRAW_DEBUG
-	if (bBuoyancyDebugDraw)
-	{
-		// Spline Color
-		const FColor SplineColor = FColor::Cyan;
-
-		// Draw projection onto the line
-		const Chaos::FVec3 SurfacePoint(ParticlePos.X, ParticlePos.Y, WaterZ);
-		Chaos::FDebugDrawQueue::GetInstance().DrawDebugLine(ParticlePos, SurfacePoint, SplineColor, false, -1.f, -1, 6.f);
-		Chaos::FDebugDrawQueue::GetInstance().DrawDebugLine(SurfacePoint, ClosestSplinePoint, SplineColor, false, -1.f, -1, 3.f);
-
-		// Draw a section of the spline near the spline key
-		Chaos::FVec3 PrevPoint;
-		bool bFirst = true;
-		for (float SplineKey = ClosestSplineKey - .1f; SplineKey <= ClosestSplineKey + .1f; SplineKey += .05f)
+		if (bBuoyancyDebugDraw)
 		{
-			const FVector SplinePoint = WaterSpline->Transform.TransformPosition(WaterSpline->Position.Eval(SplineKey));
-			if (bFirst)
-			{
-				bFirst = false;
-				PrevPoint = SplinePoint;
-			}
-			else
+			// Spline Color
+			const FColor SplineColor = FColor::Cyan;
+
+			// Draw projection onto the line
+			const Chaos::FVec3 SurfacePoint(ParticlePos.X, ParticlePos.Y, WaterZ);
+			Chaos::FDebugDrawQueue::GetInstance().DrawDebugLine(ParticlePos, SurfacePoint, SplineColor, false, -1.f, -1, 6.f);
+			Chaos::FDebugDrawQueue::GetInstance().DrawDebugLine(SurfacePoint, ClosestSplinePoint, SplineColor, false, -1.f, -1, 3.f);
+
+			// Draw a section of the spline near the spline key
+			Chaos::FVec3 PrevPoint;
+			bool bFirst = true;
+			for (float SplineKey = ClosestSplineKey - .1f; SplineKey <= ClosestSplineKey + .1f; SplineKey += .05f)
 			{
-				Chaos::FDebugDrawQueue::GetInstance().DrawDebugDirectionalArrow(PrevPoint, SplinePoint, 15.f, SplineColor, false, -1.f, -1, 3.f);
+				const FVector SplinePoint = WaterSpline.Transform.TransformPosition(WaterSpline.Position.Eval(SplineKey));
+				if (bFirst)
+				{
+					bFirst = false;
+					PrevPoint = SplinePoint;
+				}
+				else
+				{
+					Chaos::FDebugDrawQueue::GetInstance().DrawDebugDirectionalArrow(PrevPoint, SplinePoint, 15.f, SplineColor, false, -1.f, -1, 3.f);
+				}
 			}
-		}
 
-		Chaos::FDebugDrawQueue::GetInstance().DrawDebugDirectionalArrow(SurfacePoint, SurfacePoint + WaterVel, 20.f, FColor::Yellow, false, -1.f, -1, 3.f);
-	}
+			Chaos::FDebugDrawQueue::GetInstance().DrawDebugDirectionalArrow(SurfacePoint, SurfacePoint + WaterVel, 20.f, FColor::Yellow, false, -1.f, -1, 3.f);
+		}
 #endif
+	}
 
 	// Compute submerged volume and CoM
 	float SubmergedVol;
diff --git a/Engine/Plugins/Experimental/Buoyancy/Source/Runtime/Public/BuoyancyStats.h b/Engine/Plugins/Experimental/Buoyancy/Source/Runtime/Public/BuoyancyStats.h
--- a/Engine/Plugins/Experimental/Buoyancy/Source/Runtime/Public/BuoyancyStats.h
+++ b/Engine/Plugins/Experimental/Buoyancy/Source/Runtime/Public/BuoyancyStats.h
@@ -11,11 +11,14 @@ DECLARE_CYCLE_STAT(TEXT("Buoyancy_Subsystem_OnPreSimulate"), STAT_BuoyancySubsys
 DECLARE_CYCLE_STAT(TEXT("Buoyancy_Subsystem_OnMidPhaseModification"), STAT_BuoyancySubsystem_OnMidPhaseModification, STATGROUP_Buoyancy);
 DECLARE_CYCLE_STAT(TEXT("Buoyancy_Subsystem_VisitMidphases"), STAT_BuoyancySubsystem_VisitMidphases, STATGROUP_Buoyancy);
 DECLARE_CYCLE_STAT(TEXT("Buoyancy_Subsystem_ProcessMidphase"), STAT_BuoyancySubsystem_ProcessMidphase, STATGROUP_Buoyancy);
+DECLARE_CYCLE_STAT(TEXT("Buoyancy_Subsystem_SplineEvaluation"), STAT_Buoyancy_Subsystem_SplineEvaluation, STATGROUP_Buoyancy);
 DECLARE_CYCLE_STAT(TEXT("Buoyancy_Subsystem_BuildSubmersions"), STAT_BuoyancySubsystem_BuildSubmersions, STATGROUP_Buoyancy);
 DECLARE_CYCLE_STAT(TEXT("Buoyancy_Subsystem_BuildSubmersionCallbackData"), STAT_BuoyancySubsystem_BuildSubmersionCallbackData, STATGROUP_Buoyancy);
 DECLARE_CYCLE_STAT(TEXT("Buoyancy_Subsystem_ApplyBuoyantForces"), STAT_BuoyancySubsystem_ApplyBuoyantForces, STATGROUP_Buoyancy);
 DECLARE_CYCLE_STAT(TEXT("Buoyancy_Subsystem_ProduceSurfaceTouches"), STAT_BuoyancySubsystem_ProduceSurfaceTouches, STATGROUP_Buoyancy);
 DECLARE_CYCLE_STAT(TEXT("Buoyancy_Subsystem_Tick"), STAT_BuoyancySubsystem_Tick, STATGROUP_Buoyancy);
+DECLARE_CYCLE_STAT(TEXT("Buoyancy_Subsystem_UpdateBuoyancySettings"), STAT_BuoyancySubsystem_UpdateBuoyancySettings, STATGROUP_Buoyancy);
+DECLARE_CYCLE_STAT(TEXT("Buoyancy_Subsystem_UpdateWaterBodiesList"), STAT_BuoyancySubsystem_UpdateWaterBodiesList, STATGROUP_Buoyancy);
 DECLARE_CYCLE_STAT(TEXT("Buoyancy_Subsystem_DispatchCallbacks"), STAT_BuoyancySubsystem_DispatchCallbacks, STATGROUP_Buoyancy);
 
 // Algorithms
diff --git a/Engine/Plugins/Experimental/Buoyancy/Source/Runtime/Public/BuoyancySubsystem.h b/Engine/Plugins/Experimental/Buoyancy/Source/Runtime/Public/BuoyancySubsystem.h
--- a/Engine/Plugins/Experimental/Buoyancy/Source/Runtime/Public/BuoyancySubsystem.h
+++ b/Engine/Plugins/Experimental/Buoyancy/Source/Runtime/Public/BuoyancySubsystem.h
@@ -11,6 +11,7 @@
 #include "WaterBodyComponent.h"
 #include "BuoyancyEventFlags.h"
 #include "ChaosUserDataPT.h"
+#include "PhysicsProxy/SingleParticlePhysicsProxyFwd.h"
 #include "BuoyancySubsystem.generated.h"
 
 DECLARE_LOG_CATEGORY_EXTERN(LogBuoyancySubsystem, Log, All);
@@ -129,6 +130,15 @@ private:
 	void CreateSimCallback();
 	void DestroySimCallback();
 
+	// Update PT spline data structs for each waterbody in the map
+	void UpdateSplineData();
+
+	// Put updated settings struct onto async input to be sent to sim callback
+	void UpdateBuoyancySettings();
+
+	// Process async outputs which hold data for triggering callbacks
+	void ProcessSurfaceTouchCallbacks();
+
 	Chaos::FPhysicsSolver* GetSolver() const;
 
 	// When water plugin settings change, this callback will apply changes
@@ -178,11 +188,6 @@ struct FBuoyancySubmersion
 
 struct FBuoyancySubsystemSimCallbackInput : public Chaos::FSimCallbackInput
 {
-	// If this array is set, then we need to update our internal list of
-	// water body physics objects to this one. This should occur very
-	// infrequently, if more than once.
-	TOptional<TArray<Chaos::FPhysicsObjectHandle>> WaterObjects;
-
 	// If this ptr is set, then we have a new spline data manager...
 	// That should only probably happen one time
 	TOptional<FBuoyancyWaterSplineDataManager*> SplineData = nullptr;
@@ -226,13 +231,10 @@ private:
 	virtual void OnMidPhaseModification_Internal(Chaos::FMidPhaseModifierAccessor& Modifier) override;
 
 	void ProcessMidPhases(Chaos::FPBDRigidsEvolution& Evolution, Chaos::FMidPhaseModifierAccessor& MidPhaseAccessor);
-	void ProcessMidPhase(Chaos::FPBDRigidsEvolution& Evolution, Chaos::FGeometryParticleHandle* WaterParticle, Chaos::FMidPhaseModifier& MidPhase);
+	void ProcessMidPhase(Chaos::FPBDRigidsEvolution& Evolution, Chaos::FGeometryParticleHandle* WaterParticle, Chaos::FPBDRigidParticleHandle* RigidParticle, const FBuoyancyWaterSplineData& WaterSpline, Chaos::FMidPhaseModifier& MidPhase);
 	void ApplyBuoyantForces(Chaos::FPBDRigidsEvolution& Evolution);
 	void GenerateCallbackData();
 
-	// Internal array of physics objects that were created by water components
-	TArray<Chaos::FPhysicsObjectHandle> WaterObjects;
-
 	// Reference to UserDataPT sim callback which manages synchronization of
 	// water spline data
 	FBuoyancyWaterSplineDataManager* SplineData;
