Fix crash in waterbody event callbacks. Minor refactor of ProcessMidPhases for clarity

diff --git a/Engine/Plugins/Experimental/Buoyancy/Source/Runtime/Private/BuoyancySubsystem.cpp b/Engine/Plugins/Experimental/Buoyancy/Source/Runtime/Private/BuoyancySubsystem.cpp
--- a/Engine/Plugins/Experimental/Buoyancy/Source/Runtime/Private/BuoyancySubsystem.cpp
+++ b/Engine/Plugins/Experimental/Buoyancy/Source/Runtime/Private/BuoyancySubsystem.cpp
@@ -290,6 +290,13 @@ void UBuoyancySubsystem::ProcessSurfaceTouchCallbacks()
 			UPrimitiveComponent* WaterComponent = PhysScene->GetOwningComponent<UPrimitiveComponent>(SurfaceTouch.WaterProxy);
 			UPrimitiveComponent* RigidComponent = PhysScene->GetOwningComponent<UPrimitiveComponent>(SurfaceTouch.RigidProxy);
 
+			// Skip if either of the components were not valid
+			if (WaterComponent == nullptr ||
+				RigidComponent == nullptr)
+			{
+				continue;
+			}
+
 			// Get the parental water body component
 			AWaterBody* WaterActor = WaterComponent->GetOwner<AWaterBody>();
 
@@ -457,6 +464,14 @@ void FBuoyancySubsystemSimCallback::ProcessMidPhases(
 {
 	SCOPE_CYCLE_COUNTER(STAT_BuoyancySubsystem_VisitMidphases)
 
+	//
+	// This right here is the ugliest bit of the subsystem,
+	// since it loops needlessly over all midphases. It is
+	// tempting to think that this is the one most deserving
+	// of optimization, but the extra iterations here are far
+	// from the slowest part!
+	//
+
 	// Loop over all midphases
 	MidPhaseAccessor.VisitMidPhases([this, &Evolution](Chaos::FMidPhaseModifier& MidPhase)
 	{
@@ -464,30 +479,37 @@ void FBuoyancySubsystemSimCallback::ProcessMidPhases(
 		Chaos::FGeometryParticleHandle* WaterParticle;
 		Chaos::FGeometryParticleHandle* OtherParticle;
 		MidPhase.GetParticles(&WaterParticle, &OtherParticle);
-		if (WaterParticle && OtherParticle)
+		if (WaterParticle == nullptr ||
+			OtherParticle == nullptr)
 		{
-			// Get spline data for particle 0. If it exists, then it's water.
-			// If it doesn't exist, then try the other particle.
-			const TSharedPtr<FBuoyancyWaterSplineData>* WaterSpline = SplineData->GetData_PT(*WaterParticle);
+			return;
+		}
+
+		// Get spline data for particle 0. If it exists, then it's water.
+		// If it doesn't exist, then try the other particle.
+		const TSharedPtr<FBuoyancyWaterSplineData>* WaterSpline = SplineData->GetData_PT(*WaterParticle);
+		if (WaterSpline == nullptr || !WaterSpline->IsValid())
+		{
+			// Swap the particles and try again
+			Swap(WaterParticle, OtherParticle);
+			WaterSpline = SplineData->GetData_PT(*WaterParticle);
 			if (WaterSpline == nullptr || !WaterSpline->IsValid())
 			{
-				// Swap the particles and try again
-				Swap(WaterParticle, OtherParticle);
-				WaterSpline = SplineData->GetData_PT(*WaterParticle);
-				if (WaterSpline == nullptr || !WaterSpline->IsValid())
-				{
-					// Neither particle has a water spline data, so give up.
-					// This is not a water interaction
-					return;
-				}
+				// Neither particle has a water spline data, so give up.
+				// This is not a water interaction
+				return;
 			}
+		}
 
-			// Make sure the non-water particle is backed by a rigid
-			if (Chaos::FPBDRigidParticleHandle* RigidParticle = OtherParticle->CastToRigidParticle())
-			{
-				ProcessMidPhase(Evolution, WaterParticle, RigidParticle, *WaterSpline->Get(), MidPhase);
-			}
+		// Make sure the non-water particle is backed by a rigid
+		Chaos::FPBDRigidParticleHandle* RigidParticle = OtherParticle->CastToRigidParticle();
+		if (RigidParticle == nullptr)
+		{
+			return;
 		}
+
+		// Finally... we know for sure this is a midphase that we wanna process
+		ProcessMidPhase(Evolution, WaterParticle, RigidParticle, *WaterSpline->Get(), MidPhase);
 	});
 }
 
