Add netmode tracking to buoyancy subsystem and remove nullptr initialization from FBuoyancySubsystemSimCallbackInput::SplineData so that the TOptional will start out properly cleared.

diff --git a/Engine/Plugins/Experimental/Buoyancy/Source/Runtime/Private/BuoyancySubsystem.cpp b/Engine/Plugins/Experimental/Buoyancy/Source/Runtime/Private/BuoyancySubsystem.cpp
--- a/Engine/Plugins/Experimental/Buoyancy/Source/Runtime/Private/BuoyancySubsystem.cpp
+++ b/Engine/Plugins/Experimental/Buoyancy/Source/Runtime/Private/BuoyancySubsystem.cpp
@@ -107,6 +107,9 @@ void UBuoyancySubsystem::PostInitialize()
 	// Apply initial runtime settings
 	ApplyRuntimeSettings(GetDefault<UBuoyancyRuntimeSettings>(), EPropertyChangeType::ValueSet);
 
+	// Set initial net mode based on current world state
+	UpdateNetMode();
+
 	// Setup callback for when waterbodies are added/removed
 	if (FWaterBodyManager* WaterBodyManager = UWaterSubsystem::GetWaterBodyManager(GetWorld()))
 	{
@@ -196,6 +199,24 @@ void UBuoyancySubsystem::Tick(float DeltaTime)
 	}
 }
 
+void UBuoyancySubsystem::UpdateNetMode()
+{
+	// Get the netmode from the world. If it's different
+	// than the one we previously stored, send it to PT
+	const ENetMode PrevNetMode = NetMode;
+	NetMode = GetWorld()->GetNetMode();
+	if (NetMode != PrevNetMode)
+	{
+		if (SimCallback)
+		{
+			if (FBuoyancySubsystemSimCallbackInput* AsyncInput = SimCallback->GetProducerInputData_External())
+			{
+				AsyncInput->NetMode = NetMode;
+			}
+		}
+	}
+}
+
 void UBuoyancySubsystem::UpdateSplineData()
 {
 	SCOPE_CYCLE_COUNTER(STAT_BuoyancySubsystem_UpdateWaterBodiesList)
@@ -375,6 +396,7 @@ void FBuoyancySubsystemSimCallbackInput::Reset()
 {
 	SplineData.Reset();
 	BuoyancySettings.Reset();
+	NetMode.Reset();
 }
 
 void FBuoyancySubsystemSimCallbackOutput::Reset()
@@ -389,6 +411,11 @@ void FBuoyancySubsystemSimCallback::OnPreSimulate_Internal()
 	// If we were sent new buoyancy settings or data, update our local sim copy
 	if (const FBuoyancySubsystemSimCallbackInput* Input = GetConsumerInput_Internal())
 	{
+		if (Input->NetMode.IsSet())
+		{
+			NetMode = *Input->NetMode;
+		}
+
 		if (Input->SplineData.IsSet())
 		{
 			SplineData = *Input->SplineData;
@@ -411,6 +438,12 @@ void FBuoyancySubsystemSimCallback::OnMidPhaseModification_Internal(Chaos::FMidP
 {
 	SCOPE_CYCLE_COUNTER(STAT_BuoyancySubsystem_OnMidPhaseModification)
 
+	// Don't do any processing until we know what machine we're on
+	if (NetMode == ENetMode::NM_MAX)
+	{
+		return;
+	}
+
 	// If we don't have a spline data manager, early out
 	if (SplineData == nullptr)
 	{
diff --git a/Engine/Plugins/Experimental/Buoyancy/Source/Runtime/Public/BuoyancySubsystem.h b/Engine/Plugins/Experimental/Buoyancy/Source/Runtime/Public/BuoyancySubsystem.h
--- a/Engine/Plugins/Experimental/Buoyancy/Source/Runtime/Public/BuoyancySubsystem.h
+++ b/Engine/Plugins/Experimental/Buoyancy/Source/Runtime/Public/BuoyancySubsystem.h
@@ -99,6 +99,7 @@ class BUOYANCY_API UBuoyancySubsystem : public UTickableWorldSubsystem
 		, BuoyancySettings(FBuoyancySettings())
 		, SplineData(nullptr)
 		, SimCallback(nullptr)
+		, NetMode(ENetMode::NM_MAX)
 	{ }
 
 public:
@@ -130,6 +131,12 @@ private:
 	void CreateSimCallback();
 	void DestroySimCallback();
 
+	// Get netmode from world and send it to PT
+	void UpdateNetMode();
+
+	// Get the waterbodymanager pointer from the water subsystem
+	void UpdateWaterBodyManager();
+
 	// Update PT spline data structs for each waterbody in the map
 	void UpdateSplineData();
 
@@ -153,6 +160,8 @@ private:
 	FBuoyancyWaterSplineDataManager* SplineData;
 
 	class FBuoyancySubsystemSimCallback* SimCallback;
+
+	ENetMode NetMode;
 };
 
 
@@ -190,7 +199,7 @@ struct FBuoyancySubsystemSimCallbackInput : public Chaos::FSimCallbackInput
 {
 	// If this ptr is set, then we have a new spline data manager...
 	// That should only probably happen one time
-	TOptional<FBuoyancyWaterSplineDataManager*> SplineData = nullptr;
+	TOptional<FBuoyancyWaterSplineDataManager*> SplineData;
 
 	// Here we use a unique ptr so that it is possible to provide an async
 	// input _without_ buoyancy settings (which may be eventually desirable
@@ -198,6 +207,9 @@ struct FBuoyancySubsystemSimCallbackInput : public Chaos::FSimCallbackInput
 	// data).
 	mutable TUniquePtr<FBuoyancySettings> BuoyancySettings;
 
+	// Set when net mode changes - should be one time on initialization.
+	TOptional<ENetMode> NetMode;
+
 	void Reset();
 };
 
@@ -266,4 +278,8 @@ private:
 	// memory hot - the array is reset, repopulated and traversed, every frame,
 	// so we want to minimize allocations.
 	TSparseArray<TBitArray<>> SubmergedShapes;
+
+	// Used to track the net mode of the world that owns the phys scene that this
+	// sim tick is taking place in.
+	ENetMode NetMode = ENetMode::NM_MAX;
 };
