Allowing buoyancy to be enabled without impacting defaultengine.ini

diff --git a/Engine/Plugins/Experimental/Buoyancy/Source/Runtime/Private/BuoyancySubsystem.cpp b/Engine/Plugins/Experimental/Buoyancy/Source/Runtime/Private/BuoyancySubsystem.cpp
--- a/Engine/Plugins/Experimental/Buoyancy/Source/Runtime/Private/BuoyancySubsystem.cpp
+++ b/Engine/Plugins/Experimental/Buoyancy/Source/Runtime/Private/BuoyancySubsystem.cpp
@@ -68,6 +68,24 @@ bool UBuoyancySubsystem::IsEnabled() const
 	return SimCallback != nullptr;
 }
 
+bool UBuoyancySubsystem::SetEnabledWithUpdatedNetModeCallback(const bool bEnabled)
+{
+	bool bEnabledResult = SetEnabled(bEnabled);
+
+	if (IsEnabled())
+	{
+		if (SimCallback)
+		{
+			if (FBuoyancySubsystemSimCallbackInput* AsyncInput = SimCallback->GetProducerInputData_External())
+			{
+				AsyncInput->NetMode = NetMode;
+			}
+		}
+	}
+
+	return bEnabledResult;
+}
+
 void UBuoyancySubsystem::CreateSimCallback()
 {
 	// Create sim callback
diff --git a/Engine/Plugins/Experimental/Buoyancy/Source/Runtime/Public/BuoyancySubsystem.h b/Engine/Plugins/Experimental/Buoyancy/Source/Runtime/Public/BuoyancySubsystem.h
--- a/Engine/Plugins/Experimental/Buoyancy/Source/Runtime/Public/BuoyancySubsystem.h
+++ b/Engine/Plugins/Experimental/Buoyancy/Source/Runtime/Public/BuoyancySubsystem.h
@@ -112,6 +112,8 @@ public:
 	UFUNCTION()
 	bool IsEnabled() const;
 
+	bool SetEnabledWithUpdatedNetModeCallback(const bool bEnabled);
+
 protected:
 
 	// UTickableWorldSubsystem begin interface
