Fix buoyancy code not using the right shape for water in the Midphase

diff --git a/Engine/Plugins/Experimental/Buoyancy/Source/Runtime/Private/BuoyancyAlgorithms.cpp b/Engine/Plugins/Experimental/Buoyancy/Source/Runtime/Private/BuoyancyAlgorithms.cpp
--- a/Engine/Plugins/Experimental/Buoyancy/Source/Runtime/Private/BuoyancyAlgorithms.cpp
+++ b/Engine/Plugins/Experimental/Buoyancy/Source/Runtime/Private/BuoyancyAlgorithms.cpp
@@ -487,7 +487,7 @@ namespace BuoyancyAlgorithms
 		const FImplicitObject* WaterRootImplicit = WaterParticle->GetGeometry();
 		const EImplicitObjectType WaterShapeType = Chaos::Private::GetImplicitCollisionType(WaterParticle, WaterRootImplicit);
 		const FShapeInstanceArray& WaterShapeInstances = WaterParticle->ShapeInstances();
-		const FShapeInstance* WaterShapeInstance = ShapeInstances[0].Get();
+		const FShapeInstance* WaterShapeInstance = WaterShapeInstances[0].Get();
 
 		// Initialize submersion values
 		SubmergedVol = 0.f;
