Add spline width data, to determine if a body is outside a river. Add inside/outside tests for lake and river. Necessary for raised water bodies

diff --git a/Engine/Plugins/Experimental/Buoyancy/Source/Runtime/Private/BuoyancySubsystem.cpp b/Engine/Plugins/Experimental/Buoyancy/Source/Runtime/Private/BuoyancySubsystem.cpp
--- a/Engine/Plugins/Experimental/Buoyancy/Source/Runtime/Private/BuoyancySubsystem.cpp
+++ b/Engine/Plugins/Experimental/Buoyancy/Source/Runtime/Private/BuoyancySubsystem.cpp
@@ -251,18 +251,20 @@ void UBuoyancySubsystem::UpdateSplineData()
 		// Loop over every registered water body
 		WaterBodyManager->ForEachWaterBodyComponent(GetWorld(), [this](UWaterBodyComponent* WaterBodyComponent)
 		{
-			// Get the metadata object, if there is one
-			UWaterSplineMetadata* WaterSplineMetadata = WaterBodyComponent->GetWaterSplineMetadata();
-
+			// Get the metadata object if there is one, and the spline component
+			UWaterSplineMetadata* SplineMetadata = WaterBodyComponent->GetWaterSplineMetadata();
 			if (UWaterSplineComponent* SplineComponent = WaterBodyComponent->GetWaterSpline())
 			{
 				// Copy out water spline data into a shared ptr, to be associated with all
 				// child particles and marshaled to PT.
 				const Chaos::FRigidTransform3 WaterTransform = WaterBodyComponent->GetComponentTransform();
+				const TOptional<FInterpCurveFloat> EmptyOptionalFloat;
 				TSharedPtr<FBuoyancyWaterSplineData> WaterSplineData = MakeShared<FBuoyancyWaterSplineData>(
 					WaterTransform,
 					SplineComponent->SplineCurves.Position,
-					WaterSplineMetadata ? WaterSplineMetadata->WaterVelocityScalar : TOptional<FInterpCurveFloat>()
+					WaterBodyComponent->GetWaterBodyType(),
+					SplineMetadata ? SplineMetadata->RiverWidth : EmptyOptionalFloat,
+					SplineMetadata ? SplineMetadata->WaterVelocityScalar : EmptyOptionalFloat
 				);
 
 				// Go over each physics object in each primitive component which was generated
@@ -587,13 +589,48 @@ void FBuoyancySubsystemSimCallback::ProcessMidPhase(
 		const FVector ParticleLocalPos = WaterSpline.Transform.InverseTransformPosition(ParticlePos);
 		float ParticleDistance;
 		const float ClosestSplineKey = WaterSpline.Position.FindNearest(ParticleLocalPos, ParticleDistance);
-		const FVector ClosestSplinePoint = WaterSpline.Transform.TransformPosition(WaterSpline.Position.Eval(ClosestSplineKey));
-		WaterZ = ClosestSplinePoint.Z;
+		const FVector ClosestPoint = WaterSpline.Transform.TransformPosition(WaterSpline.Position.Eval(ClosestSplineKey));
+		const FVector ClosestPosDerivative = WaterSpline.Transform.TransformVector(WaterSpline.Position.EvalDerivative(ClosestSplineKey));
+		const FVector Diff = ClosestPoint - ParticlePos;
+		const FVector HorizontalDiff = FVector(Diff.X, Diff.Y, 0.f);
+
+		// Different water body types have different ways of determining
+		// whether a point is laterally inside their volume.
+		switch (WaterSpline.BodyType)
+		{
+			case EWaterBodyType::River:
+			{
+				// If distance to spline is greater than the width of the spline,
+				// then this is a river and we're outside of it.
+				if (WaterSpline.Width.IsSet())
+				{
+					const float Width = WaterSpline.Width->Eval(ClosestSplineKey);
+					const float DistSq = FVector::DotProduct(HorizontalDiff, HorizontalDiff);
+					const float WidthSq = Width * Width * .25f;
+					if (DistSq > WidthSq) { return; }
+				}
+				break;
+			}
+
+			case EWaterBodyType::Lake:
+			{
+				// Determine if we're inside the lake by projecting the horizontal spline
+				// diff onto the cross product of the spline direction and the up-vector
+				// (ie, the right-vector)
+				const FVector RightVector = FVector::CrossProduct(ClosestPosDerivative, FVector::UpVector);
+				const float DiffProj = FVector::DotProduct(RightVector, HorizontalDiff);
+				if (DiffProj < SMALL_NUMBER) { return; }
+				break;
+			}
+		}
+
+		// Get the surface z of the water at this point
+		WaterZ = ClosestPoint.Z;
 
 		// Get the water velocity at this point
 		WaterVel
 			= WaterSpline.Velocity.IsSet()
-			? WaterSpline.Velocity->Eval(ClosestSplineKey) * WaterSpline.Position.EvalDerivative(ClosestSplineKey).GetSafeNormal()
+			? WaterSpline.Velocity->Eval(ClosestSplineKey) * ClosestPosDerivative.GetSafeNormal()
 			: Chaos::FVec3::ZeroVector;
 
 #if ENABLE_DRAW_DEBUG
@@ -605,7 +642,7 @@ void FBuoyancySubsystemSimCallback::ProcessMidPhase(
 			// Draw projection onto the line
 			const Chaos::FVec3 SurfacePoint(ParticlePos.X, ParticlePos.Y, WaterZ);
 			Chaos::FDebugDrawQueue::GetInstance().DrawDebugLine(ParticlePos, SurfacePoint, SplineColor, false, -1.f, -1, 6.f);
-			Chaos::FDebugDrawQueue::GetInstance().DrawDebugLine(SurfacePoint, ClosestSplinePoint, SplineColor, false, -1.f, -1, 3.f);
+			Chaos::FDebugDrawQueue::GetInstance().DrawDebugLine(SurfacePoint, ClosestPoint, SplineColor, false, -1.f, -1, 3.f);
 
 			// Draw a section of the spline near the spline key
 			Chaos::FVec3 PrevPoint;
diff --git a/Engine/Plugins/Experimental/Buoyancy/Source/Runtime/Public/BuoyancySubsystem.h b/Engine/Plugins/Experimental/Buoyancy/Source/Runtime/Public/BuoyancySubsystem.h
--- a/Engine/Plugins/Experimental/Buoyancy/Source/Runtime/Public/BuoyancySubsystem.h
+++ b/Engine/Plugins/Experimental/Buoyancy/Source/Runtime/Public/BuoyancySubsystem.h
@@ -30,16 +30,23 @@ struct FBuoyancyWaterSplineData
 	FBuoyancyWaterSplineData(
 		const Chaos::FRigidTransform3& InTransform,
 		const FInterpCurveVector& InPosition,
+		const EWaterBodyType InBodyType,
+		const TOptional<FInterpCurveFloat>& InWidth,
 		const TOptional<FInterpCurveFloat>& InVelocity)
 		: Transform(InTransform)
 		, Position(InPosition)
+		, BodyType(InBodyType)
+		, Width(InWidth)
 		, Velocity(InVelocity)
 	{ }
 
+	// Parameters that all water bodies have
 	Chaos::FRigidTransform3 Transform;
 	FInterpCurveVector Position;
+	EWaterBodyType BodyType;
 
-	// There might not be any velocity parameter, so keep this one optional
+	// Parameters that only _some_ water bodies have
+	TOptional<FInterpCurveFloat> Width;
 	TOptional<FInterpCurveFloat> Velocity;
 };
 
