Add surface normal data to buoyancy system

diff --git a/Engine/Plugins/Experimental/Buoyancy/Source/Runtime/Private/BuoyancyAlgorithms.cpp b/Engine/Plugins/Experimental/Buoyancy/Source/Runtime/Private/BuoyancyAlgorithms.cpp
--- a/Engine/Plugins/Experimental/Buoyancy/Source/Runtime/Private/BuoyancyAlgorithms.cpp
+++ b/Engine/Plugins/Experimental/Buoyancy/Source/Runtime/Private/BuoyancyAlgorithms.cpp
@@ -373,7 +373,7 @@ namespace BuoyancyAlgorithms
 					// Compute approximate submerged bounds & update submersion
 					const float WaterZ = ShapeWorldBoundsA.Max().Z;
 					FAABB3 SubmergedBoundsInB;
-					if (ComputeSubmergedBounds(WaterZ, BoxInB, ShapeWorldTransformB, SubmergedBoundsInB))
+					if (ComputeSubmergedBounds(FVector::UpVector * WaterZ, FVector::UpVector, BoxInB, ShapeWorldTransformB, SubmergedBoundsInB))
 					{
 						// At this point we know that the shape is submerged
 						bSubmerged = true;
@@ -453,9 +453,12 @@ namespace BuoyancyAlgorithms
 		}
 	}
 
-	bool ComputeSubmergedVolume(const FPBDRigidsEvolutionGBF& Evolution, const FGeometryParticleHandle* SubmergedParticle, const FGeometryParticleHandle* WaterParticle, const float WaterZ, int32 NumSubdivisions, float MinVolume, TSparseArray<TBitArray<>>& SubmergedShapes, float& SubmergedVol, FVec3& SubmergedCoM, float& TotalVol)
+	// This variant of ComputesSubmergedVolume scales the submerged volume
+	// so as not to violate the actual volume of the object, and returns the
+	// total volume calculated.
+	bool ComputeSubmergedVolume(const FPBDRigidsEvolutionGBF& Evolution, const FGeometryParticleHandle* SubmergedParticle, const FGeometryParticleHandle* WaterParticle, const float WaterZ, const FVector& WaterN, int32 NumSubdivisions, float MinVolume, TSparseArray<TBitArray<>>& SubmergedShapes, float& SubmergedVol, FVec3& SubmergedCoM, float& TotalVol)
 	{
-		if (ComputeSubmergedVolume(SubmergedParticle, WaterParticle, WaterZ, NumSubdivisions, MinVolume, SubmergedShapes, SubmergedVol, SubmergedCoM))
+		if (ComputeSubmergedVolume(SubmergedParticle, WaterParticle, WaterZ, WaterN, NumSubdivisions, MinVolume, SubmergedShapes, SubmergedVol, SubmergedCoM))
 		{
 			ScaleSubmergedVolume(Evolution, SubmergedParticle, SubmergedVol, TotalVol);
 
@@ -474,7 +477,7 @@ namespace BuoyancyAlgorithms
 		return false;
 	}
 
-	bool ComputeSubmergedVolume(const FGeometryParticleHandle* SubmergedParticle, const FGeometryParticleHandle* WaterParticle, const float WaterZ, int32 NumSubdivisions, float MinVolume, TSparseArray<TBitArray<>>& SubmergedShapes, float& SubmergedVol, FVec3& SubmergedCoM)
+	bool ComputeSubmergedVolume(const FGeometryParticleHandle* SubmergedParticle, const FGeometryParticleHandle* WaterParticle, const float WaterZ, const FVector& WaterN, int32 NumSubdivisions, float MinVolume, TSparseArray<TBitArray<>>& SubmergedShapes, float& SubmergedVol, FVec3& SubmergedCoM)
 	{
 		// Get some initial data about the submerged particle
 		const FImplicitObject* RootImplicit = SubmergedParticle->GetGeometry();
@@ -495,7 +498,7 @@ namespace BuoyancyAlgorithms
 
 		// Traverse the submerged particle's leaves
 		RootImplicit->VisitLeafObjects(
-			[SubmergedParticle, ParticleIndex, &ShapeInstances, WaterShapeType, WaterShapeInstance, &ParticleWorldTransform, WaterZ, &NumSubdivisions, &MinVolume, &SubmergedShapes, &SubmergedVol, &SubmergedCoM]
+			[SubmergedParticle, ParticleIndex, &ShapeInstances, WaterShapeType, WaterShapeInstance, &ParticleWorldTransform, WaterZ, &WaterN, &NumSubdivisions, &MinVolume, &SubmergedShapes, &SubmergedVol, &SubmergedCoM]
 			(const FImplicitObject* Implicit, const FRigidTransform3& RelativeTransform, const int32 RootObjectIndex, const int32 ObjectIndex, const int32 LeafObjectIndex)
 		{
 			const FAABB3 RelativeBounds = Implicit->CalculateTransformedBounds(RelativeTransform);
@@ -542,8 +545,10 @@ namespace BuoyancyAlgorithms
 			for (const FAABB3& Box : SubmergedBoxes)
 			{
 				// Compute the portion of the object bounds that are submerged
+				const FVec3 ShapePos = ShapeWorldTransform.GetTranslation();
+				const FVec3 WaterPos = FVec3(ShapePos.X, ShapePos.Y, WaterZ);
 				FAABB3 SubmergedBox;
-				if (ComputeSubmergedBounds(WaterZ, Box, ShapeWorldTransform, SubmergedBox))
+				if (ComputeSubmergedBounds(WaterPos, WaterN, Box, ShapeWorldTransform, SubmergedBox))
 				{
 					// At this point we know that the shape is submerged
 					bSubmerged = true;
@@ -590,14 +595,14 @@ namespace BuoyancyAlgorithms
 		return false;
 	}
 
-	bool ComputeSubmergedBounds(float WaterZ, const FAABB3& RigidBox, const FRigidTransform3& RigidTransform, FAABB3& OutSubmergedBounds)
+	bool ComputeSubmergedBounds(const FVector& WaterX, const FVector& WaterN, const FAABB3& RigidBox, const FRigidTransform3& RigidTransform, FAABB3& OutSubmergedBounds)
 	{
 		SCOPE_CYCLE_COUNTER(STAT_BuoyancyAlgorithms_ComputeSubmergedBounds)
 
 		// Get a point and a normal direction on the water representing the surface
 		// of the water in the space of the rigid object.
-		const FVec3 SurfaceNormal = RigidTransform.InverseTransformVector(FVec3::UpVector);
-		const FVec3 SurfacePoint = RigidTransform.InverseTransformPosition(FVec3::UpVector * WaterZ);
+		const FVec3 SurfaceNormal = RigidTransform.InverseTransformVector(WaterN);
+		const FVec3 SurfacePoint = RigidTransform.InverseTransformPosition(WaterX);
 
 		// Partly submerged object can have at most 10 points intersecting
 		// with the water surface
diff --git a/Engine/Plugins/Experimental/Buoyancy/Source/Runtime/Private/BuoyancySubsystem.cpp b/Engine/Plugins/Experimental/Buoyancy/Source/Runtime/Private/BuoyancySubsystem.cpp
--- a/Engine/Plugins/Experimental/Buoyancy/Source/Runtime/Private/BuoyancySubsystem.cpp
+++ b/Engine/Plugins/Experimental/Buoyancy/Source/Runtime/Private/BuoyancySubsystem.cpp
@@ -580,6 +580,7 @@ void FBuoyancySubsystemSimCallback::ProcessMidPhase(
 
 	// Evaluate spline at the object CoM to approximate the water depth
 	float WaterZ;
+	FVector WaterN;
 	Chaos::FVec3 WaterVel;
 	{
 		SCOPE_CYCLE_COUNTER(STAT_Buoyancy_Subsystem_SplineEvaluation)
@@ -600,14 +601,23 @@ void FBuoyancySubsystemSimCallback::ProcessMidPhase(
 		{
 			case EWaterBodyType::River:
 			{
-				// If distance to spline is greater than the width of the spline,
-				// then this is a river and we're outside of it.
 				if (WaterSpline.Width.IsSet())
 				{
+					// If distance to spline is greater than the width of the spline,
+					// then this is a river and we're outside of it.
 					const float Width = WaterSpline.Width->Eval(ClosestSplineKey);
 					const float DistSq = FVector::DotProduct(HorizontalDiff, HorizontalDiff);
 					const float WidthSq = Width * Width * .25f;
 					if (DistSq > WidthSq) { return; }
+
+					// River water normal can be determined by the relationship of the 
+					// derivative of the spline position to the up vector.
+					//
+					// NOTE: This calculation breaks down in the limit of purely
+					// vertical water
+					const FVector SplineRight = FVector::CrossProduct(FVector::UpVector, ClosestPosDerivative);
+					const FVector SplineUp = FVector::CrossProduct(ClosestPosDerivative, SplineRight);
+					WaterN = SplineUp.GetSafeNormal();
 				}
 				break;
 			}
@@ -620,6 +630,10 @@ void FBuoyancySubsystemSimCallback::ProcessMidPhase(
 				const FVector RightVector = FVector::CrossProduct(ClosestPosDerivative, FVector::UpVector);
 				const float DiffProj = FVector::DotProduct(RightVector, HorizontalDiff);
 				if (DiffProj < SMALL_NUMBER) { return; }
+
+				// Lake water normal is always up
+				WaterN = FVector::UpVector;
+
 				break;
 			}
 		}
@@ -653,15 +667,19 @@ void FBuoyancySubsystemSimCallback::ProcessMidPhase(
 				if (bFirst)
 				{
 					bFirst = false;
-					PrevPoint = SplinePoint;
 				}
 				else
 				{
 					Chaos::FDebugDrawQueue::GetInstance().DrawDebugDirectionalArrow(PrevPoint, SplinePoint, 15.f, SplineColor, false, -1.f, -1, 3.f);
 				}
+				PrevPoint = SplinePoint;
 			}
 
+			// Draw water velocity at the surface point
 			Chaos::FDebugDrawQueue::GetInstance().DrawDebugDirectionalArrow(SurfacePoint, SurfacePoint + WaterVel, 20.f, FColor::Yellow, false, -1.f, -1, 3.f);
+
+			// Draw water surface normal at the surface point
+			Chaos::FDebugDrawQueue::GetInstance().DrawDebugDirectionalArrow(SurfacePoint, SurfacePoint + (WaterN * 100.f), 20.f, FColor::Green, false, -1.f, -1, 3.f);
 		}
 #endif
 	}
@@ -670,7 +688,7 @@ void FBuoyancySubsystemSimCallback::ProcessMidPhase(
 	float SubmergedVol;
 	Chaos::FVec3 SubmergedCoM;
 	float TotalVol;
-	if (BuoyancyAlgorithms::ComputeSubmergedVolume(Evolution, RigidParticle, WaterParticle, WaterZ, BuoyancySettings->MaxNumBoundsSubdivisions, BuoyancySettings->MinBoundsSubdivisionVol, SubmergedShapes, SubmergedVol, SubmergedCoM, TotalVol))
+	if (BuoyancyAlgorithms::ComputeSubmergedVolume(Evolution, RigidParticle, WaterParticle, WaterZ, WaterN, BuoyancySettings->MaxNumBoundsSubdivisions, BuoyancySettings->MinBoundsSubdivisionVol, SubmergedShapes, SubmergedVol, SubmergedCoM, TotalVol))
 	{
 		SCOPE_CYCLE_COUNTER(STAT_BuoyancySubsystem_BuildSubmersions)
 
diff --git a/Engine/Plugins/Experimental/Buoyancy/Source/Runtime/Public/BuoyancyAlgorithms.h b/Engine/Plugins/Experimental/Buoyancy/Source/Runtime/Public/BuoyancyAlgorithms.h
--- a/Engine/Plugins/Experimental/Buoyancy/Source/Runtime/Public/BuoyancyAlgorithms.h
+++ b/Engine/Plugins/Experimental/Buoyancy/Source/Runtime/Public/BuoyancyAlgorithms.h
@@ -46,14 +46,14 @@ namespace BuoyancyAlgorithms
 	bool ComputeSubmergedVolume(const Chaos::FGeometryParticleHandle* ParticleA, const Chaos::FGeometryParticleHandle* ParticleB, int32 NumSubdivisions, float MinVolume, TSparseArray<TBitArray<>>& SubmergedShapes, float& SubmergedVol, Chaos::FVec3& SubmergedCoM);
 
 	// Compute submerged volume given a single waterlevel
-	bool ComputeSubmergedVolume(const Chaos::FPBDRigidsEvolutionGBF& Evolution, const Chaos::FGeometryParticleHandle* SubmergedParticle, const Chaos::FGeometryParticleHandle* WaterParticle, const float WaterZ, int32 NumSubdivisions, float MinVolume, TSparseArray<TBitArray<>>& SubmergedShapes, float& SubmergedVol, Chaos::FVec3& SubmergedCoM, float& TotalVol);
+	bool ComputeSubmergedVolume(const Chaos::FPBDRigidsEvolutionGBF& Evolution, const Chaos::FGeometryParticleHandle* SubmergedParticle, const Chaos::FGeometryParticleHandle* WaterParticle, const float WaterZ, const FVector& WaterN, int32 NumSubdivisions, float MinVolume, TSparseArray<TBitArray<>>& SubmergedShapes, float& SubmergedVol, Chaos::FVec3& SubmergedCoM, float& TotalVol);
 
-	bool ComputeSubmergedVolume(const Chaos::FGeometryParticleHandle* SubmergedParticle, const Chaos::FGeometryParticleHandle* WaterParticle, const float WaterZ, int32 NumSubdivisions, float MinVolume, TSparseArray<TBitArray<>>& SubmergedShapes, float& SubmergedVol, Chaos::FVec3& SubmergedCoM);
+	bool ComputeSubmergedVolume(const Chaos::FGeometryParticleHandle* SubmergedParticle, const Chaos::FGeometryParticleHandle* WaterParticle, const float WaterZ, const FVector& WaterN, int32 NumSubdivisions, float MinVolume, TSparseArray<TBitArray<>>& SubmergedShapes, float& SubmergedVol, Chaos::FVec3& SubmergedCoM);
 
 	// Given an OOBB and a water level, generate another OOBB which is 1. entirely contained
 	// within the input OOBB and 2. entirely contains the portion of the OOBB which is submerged
 	// below the water level.
-	bool ComputeSubmergedBounds(float WaterZ, const Chaos::FAABB3& RigidBox, const Chaos::FRigidTransform3& RigidTransform, Chaos::FAABB3& OutSubmergedBounds);
+	bool ComputeSubmergedBounds(const FVector& WaterX, const FVector& WaterN, const Chaos::FAABB3& RigidBox, const Chaos::FRigidTransform3& RigidTransform, Chaos::FAABB3& OutSubmergedBounds);
 
 	// Given a bounds object, recursively subdivide it in eighths to a fixed maximum depth and
 	// a fixed minimum smallest subdivision volume.
