Use water surface normal as the pressure gradient for buoyancy for more realistic uneven water behavior

diff --git a/Engine/Plugins/Experimental/Buoyancy/Source/Runtime/Private/BuoyancyAlgorithms.cpp b/Engine/Plugins/Experimental/Buoyancy/Source/Runtime/Private/BuoyancyAlgorithms.cpp
--- a/Engine/Plugins/Experimental/Buoyancy/Source/Runtime/Private/BuoyancyAlgorithms.cpp
+++ b/Engine/Plugins/Experimental/Buoyancy/Source/Runtime/Private/BuoyancyAlgorithms.cpp
@@ -722,7 +722,7 @@ namespace BuoyancyAlgorithms
 		return false;
 	}
 
-	bool ComputeBuoyantForce(const FPBDRigidParticleHandle* RigidParticle, const float DeltaSeconds, const float WaterDensity, const float WaterDrag, const FVec3& GravityAccelVec, const FVec3& SubmergedCoM, const float SubmergedVol, const FVec3& WaterVel, FVec3& OutDeltaV, FVec3& OutDeltaW)
+	bool ComputeBuoyantForce(const FPBDRigidParticleHandle* RigidParticle, const float DeltaSeconds, const float WaterDensity, const float WaterDrag, const FVec3& GravityAccelVec, const FVec3& SubmergedCoM, const float SubmergedVol, const FVec3& WaterVel, const FVec3& WaterN, FVec3& OutDeltaV, FVec3& OutDeltaW)
 	{
 		SCOPE_CYCLE_COUNTER(STAT_BuoyancyAlgorithms_ComputeBuoyantForces)
 
@@ -756,7 +756,7 @@ namespace BuoyancyAlgorithms
 		const FMatrix33 WorldInvI = Chaos::Utilities::ComputeWorldSpaceInertia(RigidGeneric->RCom(), RigidGeneric->ConditionedInvI());
 
 		// Compute world buoyant force and torque
-		const FVec3 WorldForce = -GravityDir * BuoyantForce;
+		const FVec3 WorldForce = WaterN * BuoyantForce;
 		const FVec3 WorldTorque = Chaos::FVec3::CrossProduct(CoMDiff, WorldForce);
 
 		// Use inertia to convert forces to accelerations
diff --git a/Engine/Plugins/Experimental/Buoyancy/Source/Runtime/Private/BuoyancySubsystem.cpp b/Engine/Plugins/Experimental/Buoyancy/Source/Runtime/Private/BuoyancySubsystem.cpp
--- a/Engine/Plugins/Experimental/Buoyancy/Source/Runtime/Private/BuoyancySubsystem.cpp
+++ b/Engine/Plugins/Experimental/Buoyancy/Source/Runtime/Private/BuoyancySubsystem.cpp
@@ -704,19 +704,25 @@ void FBuoyancySubsystemSimCallback::ProcessMidPhase(
 				FBuoyancySubmersion& Submersion = Submersions[RigidParticleIndex];
 				ensureMsgf(Submersion.Particle == RigidParticle, TEXT("Something went wrong - there's a particle index mismatch in the Submersions sparse array"));
 
-				// Sum the volumes
-				Submersion.Vol = Submersion.Vol + SubmergedVol;
-
 				// Get the weighted-average CoM
 				// NOTE: The unchecked division should be safe since we already
 				// know SubmergedVol > SMALL_NUMBER
-				Submersion.CoM = ((Submersion.CoM * Submersion.Vol) + (SubmergedCoM * SubmergedVol)) / Submersion.Vol;
+				const float TotalSubmergedVol = Submersion.Vol + SubmergedVol;
+				Submersion.CoM = ((Submersion.CoM * Submersion.Vol) + (SubmergedCoM * SubmergedVol)) / TotalSubmergedVol;
+
+				// Sum the volumes
+				Submersion.Vol = TotalSubmergedVol;
+
+				// Get the weighted-average slerped water surface norm
+				const float VolRatio = SubmergedVol / TotalSubmergedVol;
+				const FQuat NormRot = FQuat::FindBetweenNormals(Submersion.Norm, WaterN);
+				Submersion.Norm = NormRot.Slerp(FQuat::Identity, NormRot, VolRatio).GetUpVector();
 			}
 
 			// If this particle was not yet submerged, make a new submersion for it
 			else
 			{
-				Submersions.Insert(RigidParticleIndex, { RigidParticle, SubmergedVol, SubmergedCoM, WaterVel });
+				Submersions.Insert(RigidParticleIndex, { RigidParticle, SubmergedVol, SubmergedCoM, WaterVel, WaterN });
 			}
 
 			// If this is a surface touch record it for callback, if 
@@ -776,7 +782,7 @@ void FBuoyancySubsystemSimCallback::ApplyBuoyantForces(Chaos::FPBDRigidsEvolutio
 		// Compute delta linear and angular velocities due to buoyancy. If they're big enough to
 		// matter, apply them
 		Chaos::FVec3 DeltaV, DeltaW;
-		if (BuoyancyAlgorithms::ComputeBuoyantForce(Submersion.Particle, DeltaSeconds, BuoyancySettings->WaterDensity, BuoyancySettings->WaterDrag, GravityAccel, Submersion.CoM, Submersion.Vol, Submersion.Vel, DeltaV, DeltaW))
+		if (BuoyancyAlgorithms::ComputeBuoyantForce(Submersion.Particle, DeltaSeconds, BuoyancySettings->WaterDensity, BuoyancySettings->WaterDrag, GravityAccel, Submersion.CoM, Submersion.Vol, Submersion.Vel, Submersion.Norm, DeltaV, DeltaW))
 		{
 			// Clamp delta velocities
 			DeltaV = DeltaV.GetClampedToSize(0.f, BuoyancySettings->MaxDeltaV);
diff --git a/Engine/Plugins/Experimental/Buoyancy/Source/Runtime/Public/BuoyancyAlgorithms.h b/Engine/Plugins/Experimental/Buoyancy/Source/Runtime/Public/BuoyancyAlgorithms.h
--- a/Engine/Plugins/Experimental/Buoyancy/Source/Runtime/Public/BuoyancyAlgorithms.h
+++ b/Engine/Plugins/Experimental/Buoyancy/Source/Runtime/Public/BuoyancyAlgorithms.h
@@ -61,5 +61,5 @@ namespace BuoyancyAlgorithms
 
 	// Given a rigid particle and it's submerged CoM and Volume, compute delta velocities for
 	// integrated buoyancy forces on an object
-	bool ComputeBuoyantForce(const Chaos::FPBDRigidParticleHandle* RigidParticle, const float DeltaSeconds, const float WaterDensity, const float WaterDrag, const Chaos::FVec3& GravityAccelVec, const Chaos::FVec3& SubmergedCoM, const float SubmergedVol, const Chaos::FVec3& WaterVel, Chaos::FVec3& OutDeltaV, Chaos::FVec3& OutDeltaW);
+	bool ComputeBuoyantForce(const Chaos::FPBDRigidParticleHandle* RigidParticle, const float DeltaSeconds, const float WaterDensity, const float WaterDrag, const Chaos::FVec3& GravityAccelVec, const Chaos::FVec3& SubmergedCoM, const float SubmergedVol, const Chaos::FVec3& WaterVel, const Chaos::FVec3& WaterN, Chaos::FVec3& OutDeltaV, Chaos::FVec3& OutDeltaW);
 }
diff --git a/Engine/Plugins/Experimental/Buoyancy/Source/Runtime/Public/BuoyancySubsystem.h b/Engine/Plugins/Experimental/Buoyancy/Source/Runtime/Public/BuoyancySubsystem.h
--- a/Engine/Plugins/Experimental/Buoyancy/Source/Runtime/Public/BuoyancySubsystem.h
+++ b/Engine/Plugins/Experimental/Buoyancy/Source/Runtime/Public/BuoyancySubsystem.h
@@ -202,6 +202,7 @@ struct FBuoyancySubmersion
 	float Vol;
 	FVector CoM;
 	FVector Vel;
+	FVector Norm;
 };
 
 struct FBuoyancySubsystemSimCallbackInput : public Chaos::FSimCallbackInput
