Add one more phase to buoyancy subsystem sim callback so that each particle will sort its water interactions by depth before processing them.

This sort should be very fast since these lists will only be more than 1 when crossing boundaries, and should almost never exceed 2 items.

diff --git a/Engine/Plugins/Experimental/Buoyancy/Source/Runtime/Private/BuoyancySubsystem.cpp b/Engine/Plugins/Experimental/Buoyancy/Source/Runtime/Private/BuoyancySubsystem.cpp
--- a/Engine/Plugins/Experimental/Buoyancy/Source/Runtime/Private/BuoyancySubsystem.cpp
+++ b/Engine/Plugins/Experimental/Buoyancy/Source/Runtime/Private/BuoyancySubsystem.cpp
@@ -496,13 +496,17 @@ void FBuoyancySubsystemSimCallback::OnMidPhaseModification_Internal(Chaos::FMidP
 	Swap(Submersions, PrevSubmersions);
 	Swap(SubmersionMetaData, PrevSubmersionMetaData);
 
-	// Clear submersions and submerged shapes array, but keep their memory allocated
+	// Clear arrays for the following phases, but keep their memory allocated
+	Interactions.Reset();
 	Submersions.Reset();
 	SubmergedShapes.Reset();
 	SubmersionMetaData.Reset();
 
 	// Build list of "submersions"
-	ProcessMidPhases(*Evolution, MidPhaseAccessor);
+	TrackInteractions(*Evolution, MidPhaseAccessor);
+
+	// Process the list of interactions that we built from the midphases
+	ProcessInteractions(*Evolution);
 
 	// Apply buoyant forces resulting from submersions
 	ApplyBuoyantForces(*Evolution);
@@ -511,7 +515,7 @@ void FBuoyancySubsystemSimCallback::OnMidPhaseModification_Internal(Chaos::FMidP
 	GenerateCallbackData();
 }
 
-void FBuoyancySubsystemSimCallback::ProcessMidPhases(
+void FBuoyancySubsystemSimCallback::TrackInteractions(
 	Chaos::FPBDRigidsEvolution& Evolution,
 	Chaos::FMidPhaseModifierAccessor& MidPhaseAccessor)
 {
@@ -562,62 +566,134 @@ void FBuoyancySubsystemSimCallback::ProcessMidPhases(
 		}
 
 		// Finally... we know for sure this is a midphase that we wanna process
-		ProcessMidPhase(Evolution, WaterParticle, RigidParticle, *WaterSpline->Get(), MidPhase);
+		TrackInteraction(Evolution, WaterParticle, RigidParticle, *WaterSpline->Get(), MidPhase);
 	});
 }
 
-void FBuoyancySubsystemSimCallback::ProcessMidPhase(
+void FBuoyancySubsystemSimCallback::TrackInteraction(
 	Chaos::FPBDRigidsEvolution& Evolution,
 	Chaos::FGeometryParticleHandle* WaterParticle,
 	Chaos::FPBDRigidParticleHandle* RigidParticle,
 	const FBuoyancyWaterSplineData& WaterSpline,
 	Chaos::FMidPhaseModifier& MidPhase)
 {
-	SCOPE_CYCLE_COUNTER(STAT_BuoyancySubsystem_ProcessMidphase)
+	SCOPE_CYCLE_COUNTER(STAT_BuoyancySubsystem_TrackInteraction)
 
-	// Always disable midphases with water
-	MidPhase.Disable();
+	{
+		SCOPE_CYCLE_COUNTER(STAT_BuoyancySubsystem_DisableMidPhase)
 
-	// Evaluate spline at the object CoM to approximate the water depth
-	float WaterZ;
-	FVector WaterN;
-	Chaos::FVec3 WaterVel;
+		// Always disable midphases with water
+		MidPhase.Disable();
+	}
+
+	float ClosestSplineKey;
+	FVector ClosestPoint;
 	{
-		SCOPE_CYCLE_COUNTER(STAT_Buoyancy_Subsystem_SplineEvaluation)
+		SCOPE_CYCLE_COUNTER(STAT_BuoyancySubsystem_SplineEvaluation)
 
 		// Find water surface at the nearest point on the spline
 		const FVector ParticlePos = RigidParticle->XCom();
 		const FVector ParticleLocalPos = WaterSpline.Transform.InverseTransformPosition(ParticlePos);
 		float ParticleDistance;
-		const float ClosestSplineKey = WaterSpline.Position.FindNearest(ParticleLocalPos, ParticleDistance);
-		const FVector ClosestPoint = WaterSpline.Transform.TransformPosition(WaterSpline.Position.Eval(ClosestSplineKey));
-		const FVector ClosestPosDerivative = WaterSpline.Transform.TransformVector(WaterSpline.Position.EvalDerivative(ClosestSplineKey));
-		const FVector Diff = ClosestPoint - ParticlePos;
-		const FVector HorizontalDiff = FVector(Diff.X, Diff.Y, 0.f);
+		ClosestSplineKey = WaterSpline.Position.FindNearest(ParticleLocalPos, ParticleDistance);
+		ClosestPoint = WaterSpline.Transform.TransformPosition(WaterSpline.Position.Eval(ClosestSplineKey));
+	}
+
+	{
+		SCOPE_CYCLE_COUNTER(STAT_BuoyancySubsystem_AddInteraction)
+
+		// Create the interaction struct
+		FBuoyancyInteraction BuoyancyInteraction
+		{
+			RigidParticle,
+			WaterParticle,
+			WaterSpline,
+			ClosestSplineKey,
+			ClosestPoint,
+		};
+
+		// If this particle already has a list of interactions, add to it. Otherwise
+		// make a new one.
+		const int32 RigidParticleIndex = RigidParticle->UniqueIdx().Idx;
+		if (Interactions.IsValidIndex(RigidParticleIndex) == false)
+		{
+			Interactions.Insert(RigidParticleIndex, { BuoyancyInteraction });
+		}
+		else
+		{
+			Interactions[RigidParticleIndex].Add(BuoyancyInteraction);
+		}
+	}
+}
+
+void FBuoyancySubsystemSimCallback::ProcessInteractions(Chaos::FPBDRigidsEvolution& Evolution)
+{
+	for (TArray<FBuoyancyInteraction, TInlineAllocator<MaxNumBuoyancyInteractions>>& ParticleInteractions : Interactions)
+	{
+		// Sort this particle's water interactions by water level - highest to lowest
+		ParticleInteractions.Sort([](const FBuoyancyInteraction& A, const FBuoyancyInteraction& B)
+		{
+			return A.ClosestPoint.Z > B.ClosestPoint.Z;
+		});
+
+		// Process each interaction
+		for (FBuoyancyInteraction& Interaction : ParticleInteractions)
+		{
+			ProcessInteraction(Evolution, Interaction);
+		}
+	}
+}
+
+void FBuoyancySubsystemSimCallback::ProcessInteraction(Chaos::FPBDRigidsEvolution& Evolution, FBuoyancyInteraction& Interaction)
+{
+
+	// Get water surface level and normal
+	const float WaterZ = Interaction.ClosestPoint.Z;
+	FVector WaterN;
+	FVector ClosestPosDerivative;
+	{
+		SCOPE_CYCLE_COUNTER(STAT_BuoyancySubsystem_SplineEvaluation)
+
+		// Find water surface at the nearest point on the spline
+		const FVector ParticlePos = Interaction.RigidParticle->XCom();
+		ClosestPosDerivative = Interaction.WaterSpline.Transform.TransformVector(
+			Interaction.WaterSpline.Position.EvalDerivative(Interaction.ClosestSplineKey));
+
+		// Water normal direction depends on body type
+		if (Interaction.WaterSpline.BodyType == EWaterBodyType::River)
+		{
+			// River water normal can be determined by the relationship of the 
+			// derivative of the spline position to the up vector.
+			//
+			// NOTE: This calculation breaks down in the limit of purely
+			// vertical water
+			const FVector SplineRight = FVector::CrossProduct(FVector::UpVector, ClosestPosDerivative);
+			const FVector SplineUp = FVector::CrossProduct(ClosestPosDerivative, SplineRight);
+			WaterN = SplineUp.GetSafeNormal();
+		}
+		else
+		{
+			WaterN = FVector::UpVector;
+		}
+
+		// Project the position difference onto the water surface
+		const FVector Diff = Interaction.ClosestPoint - ParticlePos;
+		const FVector LateralDiff = Diff - (WaterN * FVector::DotProduct(WaterN, Diff));
 
 		// Different water body types have different ways of determining
 		// whether a point is laterally inside their volume.
-		switch (WaterSpline.BodyType)
+		switch (Interaction.WaterSpline.BodyType)
 		{
 			case EWaterBodyType::River:
 			{
-				if (WaterSpline.Width.IsSet())
+				if (Interaction.WaterSpline.Width.IsSet())
 				{
 					// If distance to spline is greater than the width of the spline,
 					// then this is a river and we're outside of it.
-					const float Width = WaterSpline.Width->Eval(ClosestSplineKey);
-					const float DistSq = FVector::DotProduct(HorizontalDiff, HorizontalDiff);
+					const float Width = Interaction.WaterSpline.Width->Eval(Interaction.ClosestSplineKey);
+					const float DistSq = FVector::DotProduct(LateralDiff, LateralDiff);
 					const float WidthSq = Width * Width * .25f;
 					if (DistSq > WidthSq) { return; }
-
-					// River water normal can be determined by the relationship of the 
-					// derivative of the spline position to the up vector.
-					//
-					// NOTE: This calculation breaks down in the limit of purely
-					// vertical water
-					const FVector SplineRight = FVector::CrossProduct(FVector::UpVector, ClosestPosDerivative);
-					const FVector SplineUp = FVector::CrossProduct(ClosestPosDerivative, SplineRight);
-					WaterN = SplineUp.GetSafeNormal();
 				}
 				break;
 			}
@@ -628,25 +704,12 @@ void FBuoyancySubsystemSimCallback::ProcessMidPhase(
 				// diff onto the cross product of the spline direction and the up-vector
 				// (ie, the right-vector)
 				const FVector RightVector = FVector::CrossProduct(ClosestPosDerivative, FVector::UpVector);
-				const float DiffProj = FVector::DotProduct(RightVector, HorizontalDiff);
+				const float DiffProj = FVector::DotProduct(RightVector, LateralDiff);
 				if (DiffProj < SMALL_NUMBER) { return; }
-
-				// Lake water normal is always up
-				WaterN = FVector::UpVector;
-
 				break;
 			}
 		}
 
-		// Get the surface z of the water at this point
-		WaterZ = ClosestPoint.Z;
-
-		// Get the water velocity at this point
-		WaterVel
-			= WaterSpline.Velocity.IsSet()
-			? WaterSpline.Velocity->Eval(ClosestSplineKey) * ClosestPosDerivative.GetSafeNormal()
-			: Chaos::FVec3::ZeroVector;
-
 #if ENABLE_DRAW_DEBUG
 		if (bBuoyancyDebugDraw)
 		{
@@ -656,14 +719,14 @@ void FBuoyancySubsystemSimCallback::ProcessMidPhase(
 			// Draw projection onto the line
 			const Chaos::FVec3 SurfacePoint(ParticlePos.X, ParticlePos.Y, WaterZ);
 			Chaos::FDebugDrawQueue::GetInstance().DrawDebugLine(ParticlePos, SurfacePoint, SplineColor, false, -1.f, -1, 6.f);
-			Chaos::FDebugDrawQueue::GetInstance().DrawDebugLine(SurfacePoint, ClosestPoint, SplineColor, false, -1.f, -1, 3.f);
+			Chaos::FDebugDrawQueue::GetInstance().DrawDebugLine(SurfacePoint, Interaction.ClosestPoint, SplineColor, false, -1.f, -1, 3.f);
 
 			// Draw a section of the spline near the spline key
 			Chaos::FVec3 PrevPoint;
 			bool bFirst = true;
-			for (float SplineKey = ClosestSplineKey - .1f; SplineKey <= ClosestSplineKey + .1f; SplineKey += .05f)
+			for (float SplineKey = Interaction.ClosestSplineKey - .1f; SplineKey <= Interaction.ClosestSplineKey + .1f; SplineKey += .05f)
 			{
-				const FVector SplinePoint = WaterSpline.Transform.TransformPosition(WaterSpline.Position.Eval(SplineKey));
+				const FVector SplinePoint = Interaction.WaterSpline.Transform.TransformPosition(Interaction.WaterSpline.Position.Eval(SplineKey));
 				if (bFirst)
 				{
 					bFirst = false;
@@ -676,7 +739,7 @@ void FBuoyancySubsystemSimCallback::ProcessMidPhase(
 			}
 
 			// Draw water velocity at the surface point
-			Chaos::FDebugDrawQueue::GetInstance().DrawDebugDirectionalArrow(SurfacePoint, SurfacePoint + WaterVel, 20.f, FColor::Yellow, false, -1.f, -1, 3.f);
+			//Chaos::FDebugDrawQueue::GetInstance().DrawDebugDirectionalArrow(SurfacePoint, SurfacePoint + WaterVel, 20.f, FColor::Yellow, false, -1.f, -1, 3.f);
 
 			// Draw water surface normal at the surface point
 			Chaos::FDebugDrawQueue::GetInstance().DrawDebugDirectionalArrow(SurfacePoint, SurfacePoint + (WaterN * 100.f), 20.f, FColor::Green, false, -1.f, -1, 3.f);
@@ -688,7 +751,7 @@ void FBuoyancySubsystemSimCallback::ProcessMidPhase(
 	float SubmergedVol;
 	Chaos::FVec3 SubmergedCoM;
 	float TotalVol;
-	if (BuoyancyAlgorithms::ComputeSubmergedVolume(Evolution, RigidParticle, WaterParticle, WaterZ, WaterN, BuoyancySettings->MaxNumBoundsSubdivisions, BuoyancySettings->MinBoundsSubdivisionVol, SubmergedShapes, SubmergedVol, SubmergedCoM, TotalVol))
+	if (BuoyancyAlgorithms::ComputeSubmergedVolume(Evolution, Interaction.RigidParticle, Interaction.WaterParticle, WaterZ, WaterN, BuoyancySettings->MaxNumBoundsSubdivisions, BuoyancySettings->MinBoundsSubdivisionVol, SubmergedShapes, SubmergedVol, SubmergedCoM, TotalVol))
 	{
 		SCOPE_CYCLE_COUNTER(STAT_BuoyancySubsystem_BuildSubmersions)
 
@@ -696,13 +759,19 @@ void FBuoyancySubsystemSimCallback::ProcessMidPhase(
 		// approximate submerged center of mass, and apply the buoyancy force there.
 		if (SubmergedVol > SMALL_NUMBER)
 		{
-			const int32 RigidParticleIndex = RigidParticle->UniqueIdx().Idx;
+			const int32 RigidParticleIndex = Interaction.RigidParticle->UniqueIdx().Idx;
+
+			// Get the water velocity at this point
+			const FVector WaterVel
+				= Interaction.WaterSpline.Velocity.IsSet()
+				? Interaction.WaterSpline.Velocity->Eval(Interaction.ClosestSplineKey) * ClosestPosDerivative.GetSafeNormal()
+				: Chaos::FVec3::ZeroVector;
 
 			// If this particle was already marked submerged, add to its existing submersion
 			if (Submersions.IsValidIndex(RigidParticleIndex))
 			{
 				FBuoyancySubmersion& Submersion = Submersions[RigidParticleIndex];
-				ensureMsgf(Submersion.Particle == RigidParticle, TEXT("Something went wrong - there's a particle index mismatch in the Submersions sparse array"));
+				ensureMsgf(Submersion.Particle == Interaction.RigidParticle, TEXT("Something went wrong - there's a particle index mismatch in the Submersions sparse array"));
 
 				// Get the weighted-average CoM
 				// NOTE: The unchecked division should be safe since we already
@@ -717,12 +786,15 @@ void FBuoyancySubsystemSimCallback::ProcessMidPhase(
 				const float VolRatio = SubmergedVol / TotalSubmergedVol;
 				const FQuat NormRot = FQuat::FindBetweenNormals(Submersion.Norm, WaterN);
 				Submersion.Norm = NormRot.Slerp(FQuat::Identity, NormRot, VolRatio).GetUpVector();
+
+				// Blend water velocities
+				Submersion.Vel = FMath::Lerp(Submersion.Vel, WaterVel, VolRatio);
 			}
 
 			// If this particle was not yet submerged, make a new submersion for it
 			else
 			{
-				Submersions.Insert(RigidParticleIndex, { RigidParticle, SubmergedVol, SubmergedCoM, WaterVel, WaterN });
+				Submersions.Insert(RigidParticleIndex, { Interaction.RigidParticle, SubmergedVol, SubmergedCoM, WaterVel, WaterN });
 			}
 
 			// If this is a surface touch record it for callback, if 
@@ -733,8 +805,8 @@ void FBuoyancySubsystemSimCallback::ProcessMidPhase(
 
 				// Proceed only if this CoM is moving fast enough to generate events
 				const FBuoyancySubmersion& Submersion = Submersions[RigidParticleIndex];
-				const Chaos::FVec3 CoMDiff = Submersion.CoM - RigidParticle->XCom();
-				const Chaos::FVec3 CoMVel = RigidParticle->V() + Chaos::FVec3::CrossProduct(RigidParticle->W(), CoMDiff);
+				const Chaos::FVec3 CoMDiff = Submersion.CoM - Interaction.RigidParticle->XCom();
+				const Chaos::FVec3 CoMVel = Interaction.RigidParticle->V() + Chaos::FVec3::CrossProduct(Interaction.RigidParticle->W(), CoMDiff);
 				const float CoMVelSq = Chaos::FVec3::DotProduct(CoMVel, CoMVel);
 				const float MinVel = BuoyancySettings->MinVelocityForSurfaceTouchCallback;
 				const float MinVelSq = MinVel * MinVel;
@@ -751,7 +823,7 @@ void FBuoyancySubsystemSimCallback::ProcessMidPhase(
 					// If we haven't already maxed out on water contacts, add one
 					if (MetaData.WaterContacts.Num() < FBuoyancySubmersionMetaData::MaxNumWaterContacts)
 					{
-						MetaData.WaterContacts.Add({ WaterParticle, SubmergedVol, SubmergedCoM, CoMVel });
+						MetaData.WaterContacts.Add({ Interaction.WaterParticle, SubmergedVol, SubmergedCoM, CoMVel });
 					}
 				}
 			}
diff --git a/Engine/Plugins/Experimental/Buoyancy/Source/Runtime/Public/BuoyancyStats.h b/Engine/Plugins/Experimental/Buoyancy/Source/Runtime/Public/BuoyancyStats.h
--- a/Engine/Plugins/Experimental/Buoyancy/Source/Runtime/Public/BuoyancyStats.h
+++ b/Engine/Plugins/Experimental/Buoyancy/Source/Runtime/Public/BuoyancyStats.h
@@ -10,8 +10,10 @@ DECLARE_STATS_GROUP(TEXT("Buoyancy"), STATGROUP_Buoyancy, STATCAT_Advanced);
 DECLARE_CYCLE_STAT(TEXT("Buoyancy_Subsystem_OnPreSimulate"), STAT_BuoyancySubsystem_OnPreSimulate, STATGROUP_Buoyancy);
 DECLARE_CYCLE_STAT(TEXT("Buoyancy_Subsystem_OnMidPhaseModification"), STAT_BuoyancySubsystem_OnMidPhaseModification, STATGROUP_Buoyancy);
 DECLARE_CYCLE_STAT(TEXT("Buoyancy_Subsystem_VisitMidphases"), STAT_BuoyancySubsystem_VisitMidphases, STATGROUP_Buoyancy);
-DECLARE_CYCLE_STAT(TEXT("Buoyancy_Subsystem_ProcessMidphase"), STAT_BuoyancySubsystem_ProcessMidphase, STATGROUP_Buoyancy);
-DECLARE_CYCLE_STAT(TEXT("Buoyancy_Subsystem_SplineEvaluation"), STAT_Buoyancy_Subsystem_SplineEvaluation, STATGROUP_Buoyancy);
+DECLARE_CYCLE_STAT(TEXT("Buoyancy_Subsystem_TrackInteraction"), STAT_BuoyancySubsystem_TrackInteraction, STATGROUP_Buoyancy);
+DECLARE_CYCLE_STAT(TEXT("Buoyancy_Subsystem_DisableMidPhase"), STAT_BuoyancySubsystem_DisableMidPhase, STATGROUP_Buoyancy);
+DECLARE_CYCLE_STAT(TEXT("Buoyancy_Subsystem_AddInteraction"), STAT_BuoyancySubsystem_AddInteraction, STATGROUP_Buoyancy);
+DECLARE_CYCLE_STAT(TEXT("Buoyancy_Subsystem_SplineEvaluation"), STAT_BuoyancySubsystem_SplineEvaluation, STATGROUP_Buoyancy);
 DECLARE_CYCLE_STAT(TEXT("Buoyancy_Subsystem_BuildSubmersions"), STAT_BuoyancySubsystem_BuildSubmersions, STATGROUP_Buoyancy);
 DECLARE_CYCLE_STAT(TEXT("Buoyancy_Subsystem_BuildSubmersionCallbackData"), STAT_BuoyancySubsystem_BuildSubmersionCallbackData, STATGROUP_Buoyancy);
 DECLARE_CYCLE_STAT(TEXT("Buoyancy_Subsystem_ApplyBuoyantForces"), STAT_BuoyancySubsystem_ApplyBuoyantForces, STATGROUP_Buoyancy);
diff --git a/Engine/Plugins/Experimental/Buoyancy/Source/Runtime/Public/BuoyancySubsystem.h b/Engine/Plugins/Experimental/Buoyancy/Source/Runtime/Public/BuoyancySubsystem.h
--- a/Engine/Plugins/Experimental/Buoyancy/Source/Runtime/Public/BuoyancySubsystem.h
+++ b/Engine/Plugins/Experimental/Buoyancy/Source/Runtime/Public/BuoyancySubsystem.h
@@ -178,6 +178,18 @@ private:
 // Buoyancy Sim Callback
 //
 
+// Each particle will have a list of potential midphases to process,
+// which must be sorted in descending Z order. This struct is used
+// to store them
+struct FBuoyancyInteraction
+{
+	Chaos::FPBDRigidParticleHandle* RigidParticle;
+	Chaos::FGeometryParticleHandle* WaterParticle;
+	const FBuoyancyWaterSplineData& WaterSpline;
+	float ClosestSplineKey;
+	FVector ClosestPoint;
+};
+
 // Metadata for submersions, used for event callbacks
 struct FBuoyancySubmersionMetaData
 {
@@ -252,8 +264,10 @@ private:
 	virtual void OnPreSimulate_Internal() override;
 	virtual void OnMidPhaseModification_Internal(Chaos::FMidPhaseModifierAccessor& Modifier) override;
 
-	void ProcessMidPhases(Chaos::FPBDRigidsEvolution& Evolution, Chaos::FMidPhaseModifierAccessor& MidPhaseAccessor);
-	void ProcessMidPhase(Chaos::FPBDRigidsEvolution& Evolution, Chaos::FGeometryParticleHandle* WaterParticle, Chaos::FPBDRigidParticleHandle* RigidParticle, const FBuoyancyWaterSplineData& WaterSpline, Chaos::FMidPhaseModifier& MidPhase);
+	void TrackInteractions(Chaos::FPBDRigidsEvolution& Evolution, Chaos::FMidPhaseModifierAccessor& MidPhaseAccessor);
+	void TrackInteraction(Chaos::FPBDRigidsEvolution& Evolution, Chaos::FGeometryParticleHandle* WaterParticle, Chaos::FPBDRigidParticleHandle* RigidParticle, const FBuoyancyWaterSplineData& WaterSpline, Chaos::FMidPhaseModifier& MidPhase);
+	void ProcessInteractions(Chaos::FPBDRigidsEvolution& Evolution);
+	void ProcessInteraction(Chaos::FPBDRigidsEvolution& Evolution, FBuoyancyInteraction& Interaction);
 	void ApplyBuoyantForces(Chaos::FPBDRigidsEvolution& Evolution);
 	void GenerateCallbackData();
 
@@ -266,6 +280,16 @@ private:
 	// memory that was allocated by GT to minimize copies.
 	TUniquePtr<FBuoyancySettings> BuoyancySettings;
 
+	// Sparse array of arrays of buoyancy interactions - the outer array has one
+	// entry per particle, the inner array has an entry per water body that it interacts
+	// with. Each will be a very small array, sorted by Z.
+	//
+	// We use inline allocator to avoid more heap allocations, and to express the
+	// assumption that a single particle is unlikely to exceed interactions with a
+	// certain number of waterbodies at a time.
+	static constexpr int32 MaxNumBuoyancyInteractions = 2;
+	TSparseArray<TArray<FBuoyancyInteraction, TInlineAllocator<MaxNumBuoyancyInteractions>>> Interactions;
+
 	// This sparse array of submersion events is indexed on particle unique indices.
 	// All buoyant forces due to submersions are applied at once. It's stored as
 	// a member variable and reset every frame, to avoid reallocation of similarly
