Fix projected depth in buoyancy so that objects don't float out in the air on vertical water drops.

diff --git a/Engine/Plugins/Experimental/Buoyancy/Source/Runtime/Private/BuoyancyAlgorithms.cpp b/Engine/Plugins/Experimental/Buoyancy/Source/Runtime/Private/BuoyancyAlgorithms.cpp
--- a/Engine/Plugins/Experimental/Buoyancy/Source/Runtime/Private/BuoyancyAlgorithms.cpp
+++ b/Engine/Plugins/Experimental/Buoyancy/Source/Runtime/Private/BuoyancyAlgorithms.cpp
@@ -456,9 +456,9 @@ namespace BuoyancyAlgorithms
 	// This variant of ComputesSubmergedVolume scales the submerged volume
 	// so as not to violate the actual volume of the object, and returns the
 	// total volume calculated.
-	bool ComputeSubmergedVolume(const FPBDRigidsEvolutionGBF& Evolution, const FGeometryParticleHandle* SubmergedParticle, const FGeometryParticleHandle* WaterParticle, const float WaterZ, const FVector& WaterN, int32 NumSubdivisions, float MinVolume, TSparseArray<TBitArray<>>& SubmergedShapes, float& SubmergedVol, FVec3& SubmergedCoM, float& TotalVol)
+	bool ComputeSubmergedVolume(const FPBDRigidsEvolutionGBF& Evolution, const FGeometryParticleHandle* SubmergedParticle, const FGeometryParticleHandle* WaterParticle, const FVector& WaterX, const FVector& WaterN, int32 NumSubdivisions, float MinVolume, TSparseArray<TBitArray<>>& SubmergedShapes, float& SubmergedVol, FVec3& SubmergedCoM, float& TotalVol)
 	{
-		if (ComputeSubmergedVolume(SubmergedParticle, WaterParticle, WaterZ, WaterN, NumSubdivisions, MinVolume, SubmergedShapes, SubmergedVol, SubmergedCoM))
+		if (ComputeSubmergedVolume(SubmergedParticle, WaterParticle, WaterX, WaterN, NumSubdivisions, MinVolume, SubmergedShapes, SubmergedVol, SubmergedCoM))
 		{
 			ScaleSubmergedVolume(Evolution, SubmergedParticle, SubmergedVol, TotalVol);
 
@@ -477,7 +477,7 @@ namespace BuoyancyAlgorithms
 		return false;
 	}
 
-	bool ComputeSubmergedVolume(const FGeometryParticleHandle* SubmergedParticle, const FGeometryParticleHandle* WaterParticle, const float WaterZ, const FVector& WaterN, int32 NumSubdivisions, float MinVolume, TSparseArray<TBitArray<>>& SubmergedShapes, float& SubmergedVol, FVec3& SubmergedCoM)
+	bool ComputeSubmergedVolume(const FGeometryParticleHandle* SubmergedParticle, const FGeometryParticleHandle* WaterParticle, const FVector& WaterX, const FVector& WaterN, int32 NumSubdivisions, float MinVolume, TSparseArray<TBitArray<>>& SubmergedShapes, float& SubmergedVol, FVec3& SubmergedCoM)
 	{
 		// Get some initial data about the submerged particle
 		const FImplicitObject* RootImplicit = SubmergedParticle->GetGeometry();
@@ -498,7 +498,7 @@ namespace BuoyancyAlgorithms
 
 		// Traverse the submerged particle's leaves
 		RootImplicit->VisitLeafObjects(
-			[SubmergedParticle, ParticleIndex, &ShapeInstances, WaterShapeType, WaterShapeInstance, &ParticleWorldTransform, WaterZ, &WaterN, &NumSubdivisions, &MinVolume, &SubmergedShapes, &SubmergedVol, &SubmergedCoM]
+			[SubmergedParticle, ParticleIndex, &ShapeInstances, WaterShapeType, WaterShapeInstance, &ParticleWorldTransform, &WaterX, &WaterN, &NumSubdivisions, &MinVolume, &SubmergedShapes, &SubmergedVol, &SubmergedCoM]
 			(const FImplicitObject* Implicit, const FRigidTransform3& RelativeTransform, const int32 RootObjectIndex, const int32 ObjectIndex, const int32 LeafObjectIndex)
 		{
 			const FAABB3 RelativeBounds = Implicit->CalculateTransformedBounds(RelativeTransform);
@@ -544,11 +544,16 @@ namespace BuoyancyAlgorithms
 			bool bSubmerged = false;
 			for (const FAABB3& Box : SubmergedBoxes)
 			{
-				// Compute the portion of the object bounds that are submerged
+				// Get the world space position of the shape
 				const FVec3 ShapePos = ShapeWorldTransform.GetTranslation();
-				const FVec3 WaterPos = FVec3(ShapePos.X, ShapePos.Y, WaterZ);
+
+				// Get the projection of the shape position onto the water
+				const FVec3 ShapeDiff = ShapePos - WaterX;
+				const FVec3 ShapeSurfacePos = WaterX + ShapeDiff - (WaterN * FVec3::DotProduct(WaterN, ShapeDiff));
+
+				// Compute the portion of the object bounds that are submerged
 				FAABB3 SubmergedBox;
-				if (ComputeSubmergedBounds(WaterPos, WaterN, Box, ShapeWorldTransform, SubmergedBox))
+				if (ComputeSubmergedBounds(ShapeSurfacePos, WaterN, Box, ShapeWorldTransform, SubmergedBox))
 				{
 					// At this point we know that the shape is submerged
 					bSubmerged = true;
diff --git a/Engine/Plugins/Experimental/Buoyancy/Source/Runtime/Private/BuoyancySubsystem.cpp b/Engine/Plugins/Experimental/Buoyancy/Source/Runtime/Private/BuoyancySubsystem.cpp
--- a/Engine/Plugins/Experimental/Buoyancy/Source/Runtime/Private/BuoyancySubsystem.cpp
+++ b/Engine/Plugins/Experimental/Buoyancy/Source/Runtime/Private/BuoyancySubsystem.cpp
@@ -648,7 +648,6 @@ void FBuoyancySubsystemSimCallback::ProcessInteraction(Chaos::FPBDRigidsEvolutio
 {
 
 	// Get water surface level and normal
-	const float WaterZ = Interaction.ClosestPoint.Z;
 	FVector WaterN;
 	FVector ClosestPosDerivative;
 	{
@@ -717,7 +716,7 @@ void FBuoyancySubsystemSimCallback::ProcessInteraction(Chaos::FPBDRigidsEvolutio
 			const FColor SplineColor = FColor::Cyan;
 
 			// Draw projection onto the line
-			const Chaos::FVec3 SurfacePoint(ParticlePos.X, ParticlePos.Y, WaterZ);
+			const FVector SurfacePoint = Interaction.ClosestPoint - LateralDiff;
 			Chaos::FDebugDrawQueue::GetInstance().DrawDebugLine(ParticlePos, SurfacePoint, SplineColor, false, -1.f, -1, 6.f);
 			Chaos::FDebugDrawQueue::GetInstance().DrawDebugLine(SurfacePoint, Interaction.ClosestPoint, SplineColor, false, -1.f, -1, 3.f);
 
@@ -751,7 +750,7 @@ void FBuoyancySubsystemSimCallback::ProcessInteraction(Chaos::FPBDRigidsEvolutio
 	float SubmergedVol;
 	Chaos::FVec3 SubmergedCoM;
 	float TotalVol;
-	if (BuoyancyAlgorithms::ComputeSubmergedVolume(Evolution, Interaction.RigidParticle, Interaction.WaterParticle, WaterZ, WaterN, BuoyancySettings->MaxNumBoundsSubdivisions, BuoyancySettings->MinBoundsSubdivisionVol, SubmergedShapes, SubmergedVol, SubmergedCoM, TotalVol))
+	if (BuoyancyAlgorithms::ComputeSubmergedVolume(Evolution, Interaction.RigidParticle, Interaction.WaterParticle, Interaction.ClosestPoint, WaterN, BuoyancySettings->MaxNumBoundsSubdivisions, BuoyancySettings->MinBoundsSubdivisionVol, SubmergedShapes, SubmergedVol, SubmergedCoM, TotalVol))
 	{
 		SCOPE_CYCLE_COUNTER(STAT_BuoyancySubsystem_BuildSubmersions)
 
diff --git a/Engine/Plugins/Experimental/Buoyancy/Source/Runtime/Public/BuoyancyAlgorithms.h b/Engine/Plugins/Experimental/Buoyancy/Source/Runtime/Public/BuoyancyAlgorithms.h
--- a/Engine/Plugins/Experimental/Buoyancy/Source/Runtime/Public/BuoyancyAlgorithms.h
+++ b/Engine/Plugins/Experimental/Buoyancy/Source/Runtime/Public/BuoyancyAlgorithms.h
@@ -46,9 +46,9 @@ namespace BuoyancyAlgorithms
 	bool ComputeSubmergedVolume(const Chaos::FGeometryParticleHandle* ParticleA, const Chaos::FGeometryParticleHandle* ParticleB, int32 NumSubdivisions, float MinVolume, TSparseArray<TBitArray<>>& SubmergedShapes, float& SubmergedVol, Chaos::FVec3& SubmergedCoM);
 
 	// Compute submerged volume given a single waterlevel
-	bool ComputeSubmergedVolume(const Chaos::FPBDRigidsEvolutionGBF& Evolution, const Chaos::FGeometryParticleHandle* SubmergedParticle, const Chaos::FGeometryParticleHandle* WaterParticle, const float WaterZ, const FVector& WaterN, int32 NumSubdivisions, float MinVolume, TSparseArray<TBitArray<>>& SubmergedShapes, float& SubmergedVol, Chaos::FVec3& SubmergedCoM, float& TotalVol);
+	bool ComputeSubmergedVolume(const Chaos::FPBDRigidsEvolutionGBF& Evolution, const Chaos::FGeometryParticleHandle* SubmergedParticle, const Chaos::FGeometryParticleHandle* WaterParticle, const FVector& WaterX, const FVector& WaterN, int32 NumSubdivisions, float MinVolume, TSparseArray<TBitArray<>>& SubmergedShapes, float& SubmergedVol, Chaos::FVec3& SubmergedCoM, float& TotalVol);
 
-	bool ComputeSubmergedVolume(const Chaos::FGeometryParticleHandle* SubmergedParticle, const Chaos::FGeometryParticleHandle* WaterParticle, const float WaterZ, const FVector& WaterN, int32 NumSubdivisions, float MinVolume, TSparseArray<TBitArray<>>& SubmergedShapes, float& SubmergedVol, Chaos::FVec3& SubmergedCoM);
+	bool ComputeSubmergedVolume(const Chaos::FGeometryParticleHandle* SubmergedParticle, const Chaos::FGeometryParticleHandle* WaterParticle, const FVector& WaterX, const FVector& WaterN, int32 NumSubdivisions, float MinVolume, TSparseArray<TBitArray<>>& SubmergedShapes, float& SubmergedVol, Chaos::FVec3& SubmergedCoM);
 
 	// Given an OOBB and a water level, generate another OOBB which is 1. entirely contained
 	// within the input OOBB and 2. entirely contains the portion of the OOBB which is submerged
