Potential fixes for crash on load/unload/level-stream map with water bodies and buoyancy subsystem enabled.

These are low-confidence but safe fixes, with the addition of a cvar for disabling the part of water interaction callbacks that are causing the crash, in case the fixes themselves do not work.

1. In UBuoyancySubsystem::DestroySimCallback remove the SplineData callback as well as SimData callback. In UBuoyancySubsystem::CreateSimCallback, add ensures to make sure SimCallback and SplineData are both null before setting them to new sim callback ptrs - this would've caught this issue earlier.

2. Fix initial data when disabled/re-enabled by adding UpdateAllAsyncInputs call to UBuoyancySubsystem::CreateSimCallback, which will send initial state to async systems when they are first created. Remove UpdateNetMode from UBuoyancySubsystem::PostInitialize since this will now be done anyway when sim callback is created.

3. Update UpdateAsyncNetMode so that it always updates when called, regardless of whether netmode has changed. This is for consistency with other UpdateAsync functions and to make it more useful - we need to force updatees without knowing what the previously sent data was in the case of initialization.

4. Add SimCallback nullcheck to UpdateAsyncBuoyancySettings. It is not called any time that SimCallback does not exist, but it should be there anyway.

5. In FBuoyancySubsystemSimCallback::OnPreSimulate_Internal, reset all caches which contain references to water particles when SplineData sim callback changes. Probably this isn't needed, but if this code hits it means a new map has been loaded in which case all the submersion caches should get reset anyway.

6. In FBuoyancySubsystemSimCallback::GenerateCallbackData, add check for cvar bBuoyancyCallbackDataEnabled before proceeding. This should theoretically already be hotfixable behavior via the config UBuoyancyRuntimeSettings::SurfaceTouchCallbackFlags, but the cvar provides an easier way.

#swarm https://p4-swarm.epicgames.net/reviews/29339631

diff --git a/Engine/Plugins/Experimental/Buoyancy/Source/Runtime/Private/BuoyancySubsystem.cpp b/Engine/Plugins/Experimental/Buoyancy/Source/Runtime/Private/BuoyancySubsystem.cpp
--- a/Engine/Plugins/Experimental/Buoyancy/Source/Runtime/Private/BuoyancySubsystem.cpp
+++ b/Engine/Plugins/Experimental/Buoyancy/Source/Runtime/Private/BuoyancySubsystem.cpp
@@ -28,6 +28,10 @@
 // CVars
 //
 
+// Jira to remove: PLAY-21231
+bool bBuoyancyCallbackDataEnabled = true;
+FAutoConsoleVariableRef CVarBuoyancyCallbackDataEnabled(TEXT("p.Buoyancy.CallbackData.Enabled"), bBuoyancyCallbackDataEnabled, TEXT(""));
+
 #if ENABLE_DRAW_DEBUG
 bool bBuoyancyDebugDraw = false;
 FAutoConsoleVariableRef CVarBuoyancyDebugDraw(TEXT("p.Buoyancy.DebugDraw"), bBuoyancyDebugDraw, TEXT(""));
@@ -90,26 +94,39 @@ bool UBuoyancySubsystem::SetEnabledWithUpdatedNetModeCallback(const bool bEnable
 
 void UBuoyancySubsystem::CreateSimCallback()
 {
+	// Sometimes in PIE, sim callbacks have not been freed at this point.
+	// I think it is an issue with world ticking subsystem Deinitialize().
+	DestroySimCallback();
+
 	// Create sim callback
 	if (Chaos::FPhysicsSolver* Solver = GetSolver())
 	{
 		// Create the callback for keeping spline data in sync
+		ensureAlwaysMsgf(SplineData == nullptr, TEXT("UBuoyancySubsystem::CreateSimCallback: Creating new FBuoyancyWaterSplineDataManager before releasing previous SplineData"));
 		SplineData = Solver->CreateAndRegisterSimCallbackObject_External<FBuoyancyWaterSplineDataManager>();
 
 		// Create the main buoyancy sim callback
+		ensureAlwaysMsgf(SimCallback == nullptr, TEXT("UBuoyancySubsystem::CreateSimCallback: Creating new FBuoyancySubsystemSimCallback before releasing previous SimCallback"));
 		SimCallback = Solver->CreateAndRegisterSimCallbackObject_External<FBuoyancySubsystemSimCallback>();
 
 		// Give the buoyancy sim callback a reference to the spline data callback,
 		// so that it'll have access to per-particle spline data
-		if (FBuoyancySubsystemSimCallbackInput* AsyncInput = SimCallback->GetProducerInputData_External())
+		if (ensureAlwaysMsgf(SimCallback != nullptr && SplineData != nullptr, TEXT("Either SimCallback or SplineData were not properly initialized in UBuoyancySubsystem::CreateSimCallback!")))
 		{
-			AsyncInput->SplineData = SplineData;
+			if (FBuoyancySubsystemSimCallbackInput* AsyncInput = SimCallback->GetProducerInputData_External())
+			{
+				AsyncInput->SplineData = SplineData;
+			}
 		}
+
+		// Populate an initial async input so that the sim callbacks have the most up-to-date info
+		UpdateAllAsyncInputs();
 	}
 }
 
 void UBuoyancySubsystem::DestroySimCallback()
 {
+	// Destroy the main sim callback for buoyancy
 	if (SimCallback)
 	{
 		if (Chaos::FPhysicsSolverBase* Solver = SimCallback->GetSolver())
@@ -118,6 +135,16 @@ void UBuoyancySubsystem::DestroySimCallback()
 			SimCallback = nullptr;
 		}
 	}
+
+	// Destroy the sim callback for keeping spline data in sync
+	if (SplineData)
+	{
+		if (Chaos::FPhysicsSolverBase* Solver = SplineData->GetSolver())
+		{
+			Solver->UnregisterAndFreeSimCallbackObject_External(SplineData);
+			SplineData = nullptr;
+		}
+	}
 }
 
 void UBuoyancySubsystem::PostInitialize()
@@ -127,8 +154,6 @@ void UBuoyancySubsystem::PostInitialize()
 	// Apply initial runtime settings
 	ApplyRuntimeSettings(GetDefault<UBuoyancyRuntimeSettings>(), EPropertyChangeType::ValueSet);
 
-	// Set initial net mode based on current world state
-	UpdateNetMode();
 
 	// Setup callback for when waterbodies are added/removed
 	if (FWaterBodyManager* WaterBodyManager = UWaterSubsystem::GetWaterBodyManager(GetWorld()))
@@ -198,6 +223,12 @@ void UBuoyancySubsystem::Tick(float DeltaTime)
 
 	Super::Tick(DeltaTime);
 
+	UWorld* World = GetWorld();
+	if (World == nullptr)
+	{
+		return;
+	}
+
 	if (SimCallback == nullptr)
 	{
 		return;
@@ -215,6 +246,12 @@ void UBuoyancySubsystem::Tick(float DeltaTime)
 		UpdateBuoyancySettings();
 	}
 
+	if (NetMode != World->GetNetMode())
+	{
+		NetMode = World->GetNetMode();
+		UpdateNetMode();
+	}
+
 	// Process surface-touched callbacks
 	if (BuoyancySettings.SurfaceTouchCallbackFlags != 0)
 	{
@@ -222,20 +259,21 @@ void UBuoyancySubsystem::Tick(float DeltaTime)
 	}
 }
 
+void UBuoyancySubsystem::UpdateAllAsyncInputs()
+{
+	UpdateNetMode();
+	UpdateSplineData();
+	UpdateBuoyancySettings();
+}
+
 void UBuoyancySubsystem::UpdateNetMode()
 {
-	// Get the netmode from the world. If it's different
-	// than the one we previously stored, send it to PT
-	const ENetMode PrevNetMode = NetMode;
-	NetMode = GetWorld()->GetNetMode();
-	if (NetMode != PrevNetMode)
+	// Send net mode to PT
+	if (SimCallback)
 	{
-		if (SimCallback)
+		if (FBuoyancySubsystemSimCallbackInput* AsyncInput = SimCallback->GetProducerInputData_External())
 		{
-			if (FBuoyancySubsystemSimCallbackInput* AsyncInput = SimCallback->GetProducerInputData_External())
-			{
-				AsyncInput->NetMode = NetMode;
-			}
+			AsyncInput->NetMode = NetMode;
 		}
 	}
 }
@@ -299,10 +337,13 @@ void UBuoyancySubsystem::UpdateBuoyancySettings()
 {
 	SCOPE_CYCLE_COUNTER(STAT_BuoyancySubsystem_UpdateBuoyancySettings)
 
-	if (FBuoyancySubsystemSimCallbackInput* AsyncInput = SimCallback->GetProducerInputData_External())
+	if (SimCallback)
 	{
-		bBuoyancySettingsChanged = false;
-		AsyncInput->BuoyancySettings = MakeUnique<FBuoyancySettings>(BuoyancySettings);
+		if (FBuoyancySubsystemSimCallbackInput* AsyncInput = SimCallback->GetProducerInputData_External())
+		{
+			bBuoyancySettingsChanged = false;
+			AsyncInput->BuoyancySettings = MakeUnique<FBuoyancySettings>(BuoyancySettings);
+		}
 	}
 }
 
@@ -447,6 +488,19 @@ void FBuoyancySubsystemSimCallback::OnPreSimulate_Internal()
 
 			// New spline data means we need to reset our spline key cache
 			SplineKeyCache.Reset();
+
+			// New spline data also means that any internal storage of water particles
+			// might be out of date. Clear all data that potentially contains references
+			// to old water particles.
+			//
+			// NOTE: This may cause new-submersion events to trigger on already submerged
+			// objects when water splines are added/removed, depending on velocity-
+			// filtering options. It might be desirable to instead add a mechanism to
+			// suppress events for one frame instead, or to remove only the data which
+			// relate to water particles which have been removed.
+			Interactions.Reset();
+			SubmersionMetaData.Reset();
+			PrevSubmersionMetaData.Reset();
 		}
 
 		if (Input->BuoyancySettings.IsValid())
@@ -922,6 +976,11 @@ void FBuoyancySubsystemSimCallback::ApplyBuoyantForces(Chaos::FPBDRigidsEvolutio
 
 void FBuoyancySubsystemSimCallback::GenerateCallbackData()
 {
+	if (bBuoyancyCallbackDataEnabled == false)
+	{
+		return;
+	}
+
 	// Generate callback data if we're into that sort of thing
 	const uint8 CallbackFlags = BuoyancySettings->SurfaceTouchCallbackFlags;
 	if (CallbackFlags != 0)
diff --git a/Engine/Plugins/Experimental/Buoyancy/Source/Runtime/Public/BuoyancySubsystem.h b/Engine/Plugins/Experimental/Buoyancy/Source/Runtime/Public/BuoyancySubsystem.h
--- a/Engine/Plugins/Experimental/Buoyancy/Source/Runtime/Public/BuoyancySubsystem.h
+++ b/Engine/Plugins/Experimental/Buoyancy/Source/Runtime/Public/BuoyancySubsystem.h
@@ -122,12 +122,12 @@ private:
 	void CreateSimCallback();
 	void DestroySimCallback();
 
+	// Self explanatory function name :D
+	void UpdateAllAsyncInputs();
+
 	// Get netmode from world and send it to PT
 	void UpdateNetMode();
 
-	// Get the waterbodymanager pointer from the water subsystem
-	void UpdateWaterBodyManager();
-
 	// Update PT spline data structs for each waterbody in the map
 	void UpdateSplineData();
 
@@ -255,7 +255,7 @@ private:
 
 	// Reference to UserDataPT sim callback which manages synchronization of
 	// water spline data
-	FBuoyancyWaterSplineDataManager* SplineData;
+	FBuoyancyWaterSplineDataManager* SplineData = nullptr;
 
 	// Initially we won't have any settings - they have to get passed down
 	// via async input. I used TUniquePtr to control access to the same
