Add native implementation for water interactions in buoyancy subsystem

#changelist validated
#virtualized

diff --git a/Engine/Plugins/Experimental/Buoyancy/Source/Runtime/Private/BuoyancySubsystem.cpp b/Engine/Plugins/Experimental/Buoyancy/Source/Runtime/Private/BuoyancySubsystem.cpp
--- a/Engine/Plugins/Experimental/Buoyancy/Source/Runtime/Private/BuoyancySubsystem.cpp
+++ b/Engine/Plugins/Experimental/Buoyancy/Source/Runtime/Private/BuoyancySubsystem.cpp
@@ -392,30 +392,40 @@ void UBuoyancySubsystem::ProcessSurfaceTouchCallbacks()
 				// TODO: Actor relevancy check?
 
 				// If the actor implements the event interface, call the surface touched callback
-				if (Actor->Implements<UBuoyancyEventInterface>())
+				if (IBuoyancyEventInterface* InterfaceInstance = Cast<IBuoyancyEventInterface>(Actor))
 				{
 					switch (SurfaceTouch.Flag)
 					{
-						case EBuoyancyEventFlags::Begin:
-							IBuoyancyEventInterface::Execute_OnSurfaceTouchBegin(
-								Actor, WaterActor, WaterComponent, RigidComponent,
-								SurfaceTouch.Vol, SurfaceTouch.CoM, SurfaceTouch.Vel);
-							break;
-
-						case EBuoyancyEventFlags::Continue:
-							IBuoyancyEventInterface::Execute_OnSurfaceTouching(
-								Actor, WaterActor, WaterComponent, RigidComponent,
-								SurfaceTouch.Vol, SurfaceTouch.CoM, SurfaceTouch.Vel);
-							break;
-
-						case EBuoyancyEventFlags::End:
-							IBuoyancyEventInterface::Execute_OnSurfaceTouchEnd(
-								Actor, WaterActor, WaterComponent, RigidComponent,
-								SurfaceTouch.Vol, SurfaceTouch.CoM, SurfaceTouch.Vel);
-							break;
+					case EBuoyancyEventFlags::Begin:
+						InterfaceInstance->OnSurfaceTouchBegin_Native(
+							WaterComponent, RigidComponent,
+							SurfaceTouch.Vol, SurfaceTouch.CoM, SurfaceTouch.Vel);
+						IBuoyancyEventInterface::Execute_OnSurfaceTouchBegin(
+							Actor, WaterActor, WaterComponent, RigidComponent,
+							SurfaceTouch.Vol, SurfaceTouch.CoM, SurfaceTouch.Vel);
+						break;
+
+					case EBuoyancyEventFlags::Continue:
+						InterfaceInstance->OnSurfaceTouching_Native(
+							WaterComponent, RigidComponent,
+							SurfaceTouch.Vol, SurfaceTouch.CoM, SurfaceTouch.Vel);
+						IBuoyancyEventInterface::Execute_OnSurfaceTouching(
+							Actor, WaterActor, WaterComponent, RigidComponent,
+							SurfaceTouch.Vol, SurfaceTouch.CoM, SurfaceTouch.Vel);
+						break;
+
+					case EBuoyancyEventFlags::End:
+						InterfaceInstance->OnSurfaceTouchEnd_Native(
+							WaterComponent, RigidComponent,
+							SurfaceTouch.Vol, SurfaceTouch.CoM, SurfaceTouch.Vel);
+						IBuoyancyEventInterface::Execute_OnSurfaceTouchEnd(
+							Actor, WaterActor, WaterComponent, RigidComponent,
+							SurfaceTouch.Vol, SurfaceTouch.CoM, SurfaceTouch.Vel);
+						break;
 					}
 				}
 			};
+
 			DispatchEvent(WaterActor);
 			DispatchEvent(RigidComponent->GetOwner());
 		}
diff --git a/Engine/Plugins/Experimental/Buoyancy/Source/Runtime/Public/BuoyancyEventInterface.h b/Engine/Plugins/Experimental/Buoyancy/Source/Runtime/Public/BuoyancyEventInterface.h
--- a/Engine/Plugins/Experimental/Buoyancy/Source/Runtime/Public/BuoyancyEventInterface.h
+++ b/Engine/Plugins/Experimental/Buoyancy/Source/Runtime/Public/BuoyancyEventInterface.h
@@ -19,13 +19,47 @@ class BUOYANCY_API IBuoyancyEventInterface
 
 public:
 
+	/**
+	* Called when part of a submerged simulated actor first comes in
+	* contact with a water surface. This can be called for multiple parts
+	* of a complex body.
+	*/
+	virtual void OnSurfaceTouchBegin_Native(
+		UPrimitiveComponent* WaterComponent,
+		UPrimitiveComponent* SubmergedComponent,
+		float SubmergedVolume,
+		const FVector& SubmergedCenterOfMass,
+		const FVector& SubmergedVelocity) = 0;
+
+	/**
+	 * Called continually while objects maintain contact with a water surface.
+	 * May be called multiple times for different parts of an object.
+	 */
+	virtual void OnSurfaceTouching_Native(
+		UPrimitiveComponent* WaterComponent,
+		UPrimitiveComponent* SubmergedComponent,
+		float SubmergedVolume,
+		const FVector& SubmergedCenterOfMass,
+		const FVector& SubmergedVelocity) = 0;
+
+	/**
+	 * Called when a submerged body loses contact with all water surfaces. This
+	 * can result from total submersion or from coming completely out of water.
+	 */
+	virtual void OnSurfaceTouchEnd_Native(
+		UPrimitiveComponent* WaterComponent,
+		UPrimitiveComponent* SubmergedComponent,
+		float SubmergedVolume,
+		const FVector& SubmergedCenterOfMass,
+		const FVector& SubmergedVelocity) = 0;
+	
 	/**
 	 * Called when part of a submerged simulated actor first comes in
 	 * contact with a water surface. This can be called for multiple parts
 	 * of a complex body.
 	 */
 	UFUNCTION(BlueprintImplementableEvent, Category = WaterBody)
-	void  OnSurfaceTouchBegin(
+	void OnSurfaceTouchBegin(
 		class AWaterBody* WaterBodyActor,
 		UPrimitiveComponent* WaterComponent,
 		UPrimitiveComponent* SubmergedComponent,
@@ -38,7 +72,7 @@ public:
 	 * May be called multiple times for different parts of an object.
 	 */
 	UFUNCTION(BlueprintImplementableEvent, Category = WaterBody)
-	void  OnSurfaceTouching(
+	void OnSurfaceTouching(
 		class AWaterBody* WaterBodyActor,
 		UPrimitiveComponent* WaterComponent,
 		UPrimitiveComponent* SubmergedComponent,
@@ -51,11 +85,11 @@ public:
 	 * can result from total submersion or from coming completely out of water.
 	 */
 	UFUNCTION(BlueprintImplementableEvent, Category = WaterBody)
-	void  OnSurfaceTouchEnd(
+	void OnSurfaceTouchEnd(
 		class AWaterBody* WaterBodyActor,
 		UPrimitiveComponent* WaterComponent,
 		UPrimitiveComponent* SubmergedComponent,
 		float SubmergedVolume,
 		const FVector& SubmergedCenterOfMass,
 		const FVector& SubmergedVelocity);
-};
+};
\ No newline at end of file
