Add perf scope to buoyancy algorithms

diff --git a/Engine/Plugins/Experimental/Buoyancy/Source/Runtime/Private/BuoyancyAlgorithms.cpp b/Engine/Plugins/Experimental/Buoyancy/Source/Runtime/Private/BuoyancyAlgorithms.cpp
--- a/Engine/Plugins/Experimental/Buoyancy/Source/Runtime/Private/BuoyancyAlgorithms.cpp
+++ b/Engine/Plugins/Experimental/Buoyancy/Source/Runtime/Private/BuoyancyAlgorithms.cpp
@@ -40,6 +40,8 @@ namespace
 	// Mark an object's shape as submerged
 	void SubmergeShape_Internal(TSparseArray<TBitArray<>>& SubmergedShapes, const int32 ParticleIndex, const int32 ShapeIndex)
 	{
+		SCOPE_CYCLE_COUNTER(STAT_BuoyancyAlgorithms_SubmergeShapeInternal)
+
 		// If no shapes are tracked for this particle yet, add a bit array for it
 		if (SubmergedShapes.IsValidIndex(ParticleIndex) == false)
 		{
diff --git a/Engine/Plugins/Experimental/Buoyancy/Source/Runtime/Public/BuoyancyStats.h b/Engine/Plugins/Experimental/Buoyancy/Source/Runtime/Public/BuoyancyStats.h
--- a/Engine/Plugins/Experimental/Buoyancy/Source/Runtime/Public/BuoyancyStats.h
+++ b/Engine/Plugins/Experimental/Buoyancy/Source/Runtime/Public/BuoyancyStats.h
@@ -31,4 +31,5 @@ DECLARE_CYCLE_STAT(TEXT("Buoyancy_Algorithms_ComputeSubmergedBounds"), STAT_Buoy
 DECLARE_CYCLE_STAT(TEXT("Buoyancy_Algorithms_ScaleSubmergedVolume"), STAT_BuoyancyAlgorithms_ScaleSubmergedVolume, STATGROUP_Buoyancy);
 DECLARE_CYCLE_STAT(TEXT("Buoyancy_Algorithms_SubdivideBounds"), STAT_BuoyancyAlgorithms_SubdivideBounds, STATGROUP_Buoyancy);
 DECLARE_CYCLE_STAT(TEXT("Buoyancy_Algorithms_ComputeBuoyantForces"), STAT_BuoyancyAlgorithms_ComputeBuoyantForces, STATGROUP_Buoyancy);
+DECLARE_CYCLE_STAT(TEXT("Buoyancy_Algorithms_SubmergeShapeInternal"), STAT_BuoyancyAlgorithms_SubmergeShapeInternal, STATGROUP_Buoyancy);
 
