Sphere buoyancy fixes
- sphere submerged volume no longer depends on orientation and won't add torque

This is a temporary solution until we add shape type handling in the buoyancy submerged volume calculation.

diff --git a/Engine/Plugins/Experimental/Buoyancy/Source/Runtime/Private/BuoyancyAlgorithms.cpp b/Engine/Plugins/Experimental/Buoyancy/Source/Runtime/Private/BuoyancyAlgorithms.cpp
--- a/Engine/Plugins/Experimental/Buoyancy/Source/Runtime/Private/BuoyancyAlgorithms.cpp
+++ b/Engine/Plugins/Experimental/Buoyancy/Source/Runtime/Private/BuoyancyAlgorithms.cpp
@@ -11,6 +11,7 @@
 #include "Chaos/PBDRigidsEvolutionGBF.h"
 #include "Chaos/Collision/CollisionFilter.h"
 #include "Chaos/Collision/CollisionUtil.h"
+#include "Chaos/Sphere.h"
 
 //
 // CVars
@@ -127,6 +128,7 @@ namespace BuoyancyAlgorithms
 	using Chaos::FShapeInstanceArray;
 	using Chaos::EImplicitObjectType;
 	using Chaos::FConstGenericParticleHandle;
+	using Chaos::FImplicitSphere3;
 
 	FRealSingle ComputeParticleVolume(const FPBDRigidsEvolutionGBF& Evolution, const FGeometryParticleHandle* Particle)
 	{
@@ -320,8 +322,19 @@ namespace BuoyancyAlgorithms
 			}
 
 			// Get the world-space bounds of shape A
-			const FRigidTransform3 ShapeWorldTransform = RelativeTransform * ParticleWorldTransform;
-			const FAABB3 LocalBox = Implicit->BoundingBox();
+			FRigidTransform3 ShapeWorldTransform = RelativeTransform * ParticleWorldTransform;
+			FAABB3 LocalBox = Implicit->BoundingBox();
+			if (const FImplicitSphere3* Sphere = Implicit->AsA<FImplicitSphere3>())
+			{
+				// If we have a sphere, ignore rotation because submerged volume is independent of rotation
+				// and also we don't want to apply any torques on the wheel.
+				// @todo(chaos): ComputeSubmergedBounds special case for spheres
+				const FVec3 SphereCenter = ShapeWorldTransform.TransformPosition(Sphere->GetCenter());
+				const FVec3 SphereExtent = FVec3(Sphere->GetRadius());
+				LocalBox = FAABB3(-SphereExtent, SphereExtent);
+				ShapeWorldTransform.SetTranslation(SphereCenter);
+				ShapeWorldTransform.SetRotation(FRotationMatrix::MakeFromZ(WaterN).ToQuat());
+			}
 			const FAABB3 WorldBox = LocalBox.TransformedAABB(ShapeWorldTransform);
 
 #if ENABLE_DRAW_DEBUG
