Add checks for null, proxy, and pending delete to particle handles that are dereferenced for the purpose of buoyancy callbacks

diff --git a/Engine/Plugins/Experimental/Buoyancy/Source/Runtime/Private/BuoyancySubsystem.cpp b/Engine/Plugins/Experimental/Buoyancy/Source/Runtime/Private/BuoyancySubsystem.cpp
--- a/Engine/Plugins/Experimental/Buoyancy/Source/Runtime/Private/BuoyancySubsystem.cpp
+++ b/Engine/Plugins/Experimental/Buoyancy/Source/Runtime/Private/BuoyancySubsystem.cpp
@@ -32,6 +32,10 @@
 bool bBuoyancyCallbackDataEnabled = true;
 FAutoConsoleVariableRef CVarBuoyancyCallbackDataEnabled(TEXT("p.Buoyancy.CallbackData.Enabled"), bBuoyancyCallbackDataEnabled, TEXT(""));
 
+// Added for jira PLAY-37027
+bool bBuoyancyCallbackDataParticleValidation = true;
+FAutoConsoleVariableRef CVarBuoyancyCallbackDataParticleValidation(TEXT("p.Buoyancy.CallbackData.ParticleValidation"), bBuoyancyCallbackDataParticleValidation, TEXT(""));
+
 #if ENABLE_DRAW_DEBUG
 bool bBuoyancyDebugDraw = false;
 FAutoConsoleVariableRef CVarBuoyancyDebugDraw(TEXT("p.Buoyancy.DebugDraw"), bBuoyancyDebugDraw, TEXT(""));
@@ -984,6 +988,35 @@ void FBuoyancySubsystemSimCallback::ApplyBuoyantForces(Chaos::FPBDRigidsEvolutio
 	}
 }
 
+namespace
+{
+	bool IsParticleValid(Chaos::FGeometryParticleHandle* ParticleHandle)
+	{
+		if (bBuoyancyCallbackDataParticleValidation == false)
+		{
+			return true;
+		}
+
+		if (ParticleHandle == nullptr)
+		{
+			return false;
+		}
+
+		IPhysicsProxyBase* Proxy = ParticleHandle->PhysicsProxy();
+		if (Proxy == nullptr)
+		{
+			return false;
+		}
+
+		if (Proxy->GetMarkedDeleted())
+		{
+			return false;
+		}
+
+		return true;
+	};
+}
+
 void FBuoyancySubsystemSimCallback::GenerateCallbackData()
 {
 	if (bBuoyancyCallbackDataEnabled == false)
@@ -1033,9 +1066,19 @@ void FBuoyancySubsystemSimCallback::GenerateCallbackData()
 				continue;
 			}
 
+			if (!ensureMsgf(IsParticleValid(Submersion.Particle), TEXT("Submersion data for buoyancy callback includes invalid submerged particle handle")))
+			{
+				continue;
+			}
+
 			// Build up output of new and continuing surface touches
 			for (const FBuoyancySubmersionMetaData::FWaterContact& WaterContact : MetaData.WaterContacts)
 			{
+				if (!ensureMsgf(IsParticleValid(WaterContact.Water), TEXT("Submersion data for buoyancy callback includes invalid water body particle handle")))
+				{
+					continue;
+				}
+
 				Output.SurfaceTouches.Add({
 					TouchFlag,
 					Submersion.Particle->PhysicsProxy(),
@@ -1059,8 +1102,19 @@ void FBuoyancySubsystemSimCallback::GenerateCallbackData()
 				const FBuoyancySubmersionMetaData& MetaData = *Iter;
 				const int32 ObjectIndex = Iter.GetIndex();
 				const FBuoyancySubmersion& Submersion = PrevSubmersions[ObjectIndex];
+
+				if (!ensureMsgf(IsParticleValid(Submersion.Particle), TEXT("Previous frame submersion data for buoyancy callback includes invalid submerged particle handle")))
+				{
+					continue;
+				}
+
 				for (const FBuoyancySubmersionMetaData::FWaterContact& WaterContact : MetaData.WaterContacts)
 				{
+					if (!ensureMsgf(IsParticleValid(WaterContact.Water), TEXT("Previous frame submersion data for buoyancy callback includes invalid water body particle handle")))
+					{
+						continue;
+					}
+
 					Output.SurfaceTouches.Add({
 						EBuoyancyEventFlags::End,
 						Submersion.Particle->PhysicsProxy(),
