FBuoyancyParticleData: Do management of per-particle data for buoyancy subsystem inside this structure so that allocation can be done in using alternative data structures. p.Buoyancy.ParticleData.HashIndex enables a lower memory footprint mode. Disabling that cvar leaves original underlying data structures.
FHashTable: Add GetAllocatedSize to FHashTable, and also make Get functions const

diff --git a/Engine/Plugins/Experimental/Buoyancy/Source/Runtime/Private/BuoyancyAlgorithms.cpp b/Engine/Plugins/Experimental/Buoyancy/Source/Runtime/Private/BuoyancyAlgorithms.cpp
--- a/Engine/Plugins/Experimental/Buoyancy/Source/Runtime/Private/BuoyancyAlgorithms.cpp
+++ b/Engine/Plugins/Experimental/Buoyancy/Source/Runtime/Private/BuoyancyAlgorithms.cpp
@@ -2,6 +2,7 @@
 
 #include "BuoyancyAlgorithms.h"
 #include "BuoyancyStats.h"
+#include "BuoyancyParticleData.h"
 #include "DrawDebugHelpers.h"
 #include "Engine/Engine.h"
 #include "Chaos/ParticleHandle.h"
@@ -256,11 +257,11 @@ namespace BuoyancyAlgorithms
 	}
 
 
-	bool ComputeSubmergedVolume(const FPBDRigidsEvolutionGBF& Evolution, const FGeometryParticleHandle* SubmergedParticle, const FGeometryParticleHandle* WaterParticle, const FVector& WaterX, const FVector& WaterN, int32 NumSubdivisions, float MinVolume, TSparseArray<TBitArray<>>& SubmergedShapes, float& SubmergedVol, FVec3& SubmergedCoM, float& TotalVol)
+	bool ComputeSubmergedVolume(FBuoyancyParticleData& ParticleData, const FPBDRigidsEvolutionGBF& Evolution, const FGeometryParticleHandle* SubmergedParticle, const FGeometryParticleHandle* WaterParticle, const FVector& WaterX, const FVector& WaterN, int32 NumSubdivisions, float MinVolume, float& SubmergedVol, FVec3& SubmergedCoM, float& TotalVol)
 	{
 		SCOPE_CYCLE_COUNTER(STAT_BuoyancyAlgorithms_ComputeSubmergedVolume)
 
-		if (ComputeSubmergedVolume(SubmergedParticle, WaterParticle, WaterX, WaterN, NumSubdivisions, MinVolume, SubmergedShapes, SubmergedVol, SubmergedCoM))
+		if (ComputeSubmergedVolume(ParticleData, SubmergedParticle, WaterParticle, WaterX, WaterN, NumSubdivisions, MinVolume, SubmergedVol, SubmergedCoM))
 		{
 			ScaleSubmergedVolume(Evolution, SubmergedParticle, SubmergedVol, TotalVol);
 
@@ -279,14 +280,15 @@ namespace BuoyancyAlgorithms
 		return false;
 	}
 
-	bool ComputeSubmergedVolume(const FGeometryParticleHandle* SubmergedParticle, const FGeometryParticleHandle* WaterParticle, const FVector& WaterX, const FVector& WaterN, int32 NumSubdivisions, float MinVolume, TSparseArray<TBitArray<>>& SubmergedShapes, float& SubmergedVol, FVec3& SubmergedCoM)
+	bool ComputeSubmergedVolume(FBuoyancyParticleData& ParticleData, const FGeometryParticleHandle* SubmergedParticle, const FGeometryParticleHandle* WaterParticle, const FVector& WaterX, const FVector& WaterN, int32 NumSubdivisions, float MinVolume, float& SubmergedVol, FVec3& SubmergedCoM)
 	{
 		// Get some initial data about the submerged particle
 		const FImplicitObject* RootImplicit = SubmergedParticle->GetGeometry();
 		const FShapeInstanceArray& ShapeInstances = SubmergedParticle->ShapeInstances();
 		const FConstGenericParticleHandle SubmergedGeneric = SubmergedParticle;
 		const FRigidTransform3 ParticleWorldTransform = SubmergedGeneric->GetTransformPQ();
-		const int32 ParticleIndex = SubmergedParticle->UniqueIdx().Idx;
+		const int32 ParticleIndex = ParticleData.GetIndex(*SubmergedParticle);
+		TSparseArray<TBitArray<>>& SubmergedShapes = ParticleData.SubmergedShapes;
 
 		// Some info about the water
 		const FImplicitObject* WaterRootImplicit = WaterParticle->GetGeometry();
diff --git a/Engine/Plugins/Experimental/Buoyancy/Source/Runtime/Private/BuoyancyParticleData.cpp b/Engine/Plugins/Experimental/Buoyancy/Source/Runtime/Private/BuoyancyParticleData.cpp
new file mode 100644
--- /dev/null
+++ b/Engine/Plugins/Experimental/Buoyancy/Source/Runtime/Private/BuoyancyParticleData.cpp
@@ -0,0 +1,192 @@
+// Copyright Epic Games, Inc. All Rights Reserved.
+
+#include "BuoyancyParticleData.h"
+
+bool bBuoyancyParticleDataHashIndex = true;
+FAutoConsoleVariableRef CVarBuoyancyParticleDataHashIndex(TEXT("p.Buoyancy.ParticleData.HashIndex"), bBuoyancyParticleDataHashIndex, TEXT("If true, use an FHashTable to index particle data maps."));
+
+float BuoyancyParticleDataHashSizeRatioGrowThreshold = 1.5f;
+FAutoConsoleVariableRef CVarBuoyancyParticleDataHashSizeRatioGrowThreshold(TEXT("p.Buoyancy.ParticleData.HashSizeRatioGrowThreshold"), BuoyancyParticleDataHashSizeRatioGrowThreshold, TEXT("If the ratio of index size to hash size passes this threshold, a rehash will occur to optimize memory. Only occurs if bBuoyancyParticleDataHashIndex is true. A value of less than 1 will result in a sparser hash table, a value of greater than 1 results in a denser table. A value of 0 or less will cause the optimization to be skipped. No value can guarantee zero hash collisions."));
+
+int32 BuoyancyParticleDataMinHashSize = 1024;
+FAutoConsoleVariableRef CVarBuoyancyParticleDataMinHashSize(TEXT("p.Buoyancy.ParticleData.MinHashSize"), BuoyancyParticleDataMinHashSize, TEXT("Hash size is not allowed to shrink below this size. Must be a power of 2!"));
+
+
+int32 BuoyancyParticleDataMaxHashSize = 8192;
+FAutoConsoleVariableRef CVarBuoyancyParticleDataMaxHashSize(TEXT("p.Buoyancy.ParticleData.MaxHashSize"), BuoyancyParticleDataMaxHashSize, TEXT("Hash size is not allowed to grow beyond this size. Must be a power of 2!"));
+
+FBuoyancyParticleData::FBuoyancyParticleData()
+	: IndexMap(BuoyancyParticleDataMinHashSize)
+{ }
+
+FBuoyancyParticleData::~FBuoyancyParticleData()
+{ }
+
+void FBuoyancyParticleData::Reset()
+{
+	SCOPE_CYCLE_COUNTER(STAT_BuoyancyParticleData_Reset)
+	LLM_SCOPE_BYTAG(BuoyancyParticleDataTag);
+
+	IndexMap.Clear();
+	ReverseIndexMap.Reset();
+	Interactions.Reset();
+	Submersions.Reset();
+	PrevSubmersions.Reset();
+	SubmersionMetaData.Reset();
+	PrevSubmersionMetaData.Reset();
+	SubmergedShapes.Reset();
+}
+
+SIZE_T FBuoyancyParticleData::GetAllocatedSize() const
+{
+	return
+		IndexMap.GetAllocatedSize() +
+		ReverseIndexMap.GetAllocatedSize() +
+		Interactions.GetAllocatedSize() +
+		Submersions.GetAllocatedSize() + 
+		PrevSubmersions.GetAllocatedSize() +
+		SubmersionMetaData.GetAllocatedSize() +
+		PrevSubmersionMetaData.GetAllocatedSize() +
+		SubmergedShapes.GetAllocatedSize();
+}
+
+int32 FBuoyancyParticleData::GetIndex(const Chaos::FGeometryParticleHandle& ParticleHandle)
+{
+	int32 ParticleIndex, ParticleKey;
+	return GetIndex(ParticleHandle, ParticleIndex, ParticleKey);
+}
+
+int32 FBuoyancyParticleData::GetIndex(const Chaos::FGeometryParticleHandle& ParticleHandle, int32& OutParticleIndex, int32& OutParticleKey)
+{
+	OutParticleIndex = ParticleHandle.UniqueIdx().Idx;
+
+	// If we're not hashing indices here, return the particle index itself
+	if (bBuoyancyParticleDataHashIndex == false)
+	{
+		OutParticleKey = INDEX_NONE;
+		return OutParticleIndex;
+	}
+
+	OutParticleKey = MurmurFinalize32(OutParticleIndex);
+	return GetIndex(OutParticleIndex, OutParticleKey);
+}
+
+int32 FBuoyancyParticleData::GetIndex(const int32 ParticleIndex, const int32 ParticleKey)
+{
+	for (uint32 DataIndex = IndexMap.First(ParticleKey); IndexMap.IsValid(DataIndex); DataIndex = IndexMap.Next(DataIndex))
+	{
+		if (ReverseIndexMap[DataIndex] == ParticleIndex)
+		{
+			return DataIndex;
+		}
+	}
+
+	return INDEX_NONE;
+}
+
+int32 FBuoyancyParticleData::AddOrGetIndex(const Chaos::FGeometryParticleHandle& ParticleHandle)
+{
+	LLM_SCOPE_BYTAG(BuoyancyParticleDataTag);
+
+	const int32 ParticleIndex = ParticleHandle.UniqueIdx().Idx;
+
+	// If we're not hashing indices here, return the particle index itself
+	if (bBuoyancyParticleDataHashIndex == false)
+	{
+		return ParticleIndex;
+	}
+
+	const int32 ParticleKey = MurmurFinalize32(ParticleIndex);
+	int32 DataIndex = GetIndex(ParticleIndex, ParticleKey);
+	if (DataIndex == INDEX_NONE)
+	{
+		DataIndex = ReverseIndexMap.Add(ParticleIndex);
+		IndexMap.Add(ParticleKey, DataIndex);
+		OptimizeMemory();
+	}
+
+	return DataIndex;
+}
+
+bool FBuoyancyParticleData::RemoveIndex(const Chaos::FGeometryParticleHandle& ParticleHandle)
+{
+	LLM_SCOPE_BYTAG(BuoyancyParticleDataTag);
+
+	int32 ParticleIndex, ParticleKey;
+	const int32 DataIndex = GetIndex(ParticleHandle, ParticleIndex, ParticleKey);
+
+	// Remove the key from our index map
+	IndexMap.Remove(ParticleKey, DataIndex);
+
+	// Remove the data index from all sparse arrays
+	ReverseIndexMap.RemoveAt(DataIndex);
+	Interactions.RemoveAt(DataIndex);
+	Submersions.RemoveAt(DataIndex);
+	PrevSubmersions.RemoveAt(DataIndex);
+	SubmersionMetaData.RemoveAt(DataIndex);
+	PrevSubmersionMetaData.RemoveAt(DataIndex);
+	SubmergedShapes.RemoveAt(DataIndex);
+
+	OptimizeMemory();
+
+	return true;
+}
+
+void FBuoyancyParticleData::OptimizeMemory()
+{
+	SCOPE_CYCLE_COUNTER(STAT_BuoyancyParticleData_OptimizeMemory)
+	LLM_SCOPE_BYTAG(BuoyancyParticleDataTag);
+
+	// Shrink all maps to fit data
+	ReverseIndexMap.Shrink();
+	Interactions.Shrink();
+	Submersions.Shrink();
+	PrevSubmersions.Shrink();
+	SubmersionMetaData.Shrink();
+	PrevSubmersionMetaData.Shrink();
+	SubmergedShapes.Shrink();
+
+	if (bBuoyancyParticleDataHashIndex == false)
+	{
+		return;
+	}
+
+	const float HashSizeRatioThreshold = BuoyancyParticleDataHashSizeRatioGrowThreshold;
+	if (HashSizeRatioThreshold <= 0.f)
+	{
+		return;
+	}
+
+	// Compute the actual hash size ratio - if it hasn't passed the threshold,
+	// nothing to do.
+	const float IndexSize = IndexMap.GetIndexSize();
+	const float HashSize = IndexMap.GetHashSize();
+	const float HashSizeRatio = IndexSize / FMath::Max(1.f, HashSize);
+	if (HashSizeRatio <= HashSizeRatioThreshold)
+	{
+		return;
+	}
+
+	// Find the smallest power of two which is greater than IndexSize. If it didn't
+	// change, do nothing.
+	const int32 MinHashSize = BuoyancyParticleDataMinHashSize;
+	const int32 MaxHashSize = BuoyancyParticleDataMaxHashSize;
+	const int32 NewRawHashSize = FMath::Pow(2, FMath::CeilToFloat(FMath::Log2(IndexSize)));
+	const int32 NewHashSize = FMath::Clamp(NewRawHashSize, MinHashSize, MaxHashSize);
+	if (HashSize == NewHashSize)
+	{
+		return;
+	}
+
+	// If we have a new hash size, reallocate our IndexMap with our new hash size, and
+	// repopulate it using the reverse index map data.
+	IndexMap.Clear(NewHashSize, IndexSize);
+	for (auto It = ReverseIndexMap.CreateConstIterator(); It; ++It)
+	{
+		const int32 DataIndex = It.GetIndex();
+		const int32 ParticleIndex = *It;
+		const int32 ParticleKey = MurmurFinalize32(ParticleIndex);
+		IndexMap.Add(ParticleKey, DataIndex);
+	}
+}
+
diff --git a/Engine/Plugins/Experimental/Buoyancy/Source/Runtime/Private/BuoyancyStats.cpp b/Engine/Plugins/Experimental/Buoyancy/Source/Runtime/Private/BuoyancyStats.cpp
new file mode 100644
--- /dev/null
+++ b/Engine/Plugins/Experimental/Buoyancy/Source/Runtime/Private/BuoyancyStats.cpp
@@ -0,0 +1,6 @@
+// Copyright Epic Games, Inc. All Rights Reserved.
+
+#include "BuoyancyStats.h"
+
+LLM_DEFINE_TAG(BuoyancyParticleDataTag, TEXT("BuoyancyParticleData"), TEXT("Physics"));
+
diff --git a/Engine/Plugins/Experimental/Buoyancy/Source/Runtime/Private/BuoyancySubsystem.cpp b/Engine/Plugins/Experimental/Buoyancy/Source/Runtime/Private/BuoyancySubsystem.cpp
--- a/Engine/Plugins/Experimental/Buoyancy/Source/Runtime/Private/BuoyancySubsystem.cpp
+++ b/Engine/Plugins/Experimental/Buoyancy/Source/Runtime/Private/BuoyancySubsystem.cpp
@@ -49,6 +49,25 @@ FAutoConsoleVariableRef CVarBuoyancyDebugDraw(TEXT("p.Buoyancy.DebugDraw"), bBuo
 DEFINE_LOG_CATEGORY(LogBuoyancySubsystem);
 
 
+//
+// Console Commands
+//
+
+#if WITH_BUOYANCY_MEMORY_TRACKING
+static FAutoConsoleCommandWithWorld BuoyancyLogMemory(
+	TEXT("p.Buoyancy.LogMemory"),
+	TEXT("Log how much memory is being used by the buoyancy subsystem"),
+	FConsoleCommandWithWorldDelegate::CreateLambda([](UWorld* World)
+	{
+		if (UBuoyancySubsystem* BuoyancySubsystem = World != nullptr ? World->GetSubsystem<UBuoyancySubsystem>() : nullptr)
+		{
+			const int32 AllocatedSize = BuoyancySubsystem->GetAllocatedSize();
+			UE_LOG(LogBuoyancySubsystem, Warning, TEXT("Buoyancy Subsystem Allocated Bytes: %d"), AllocatedSize);
+		}
+	}));
+#endif // WITH_BUOYANCY_MEMORY_TRACKING
+
+
 //
 // Buoyancy Subsystem
 //
@@ -96,6 +115,13 @@ bool UBuoyancySubsystem::SetEnabledWithUpdatedNetModeCallback(const bool bEnable
 	return bEnabledResult;
 }
 
+#if WITH_BUOYANCY_MEMORY_TRACKING
+SIZE_T UBuoyancySubsystem::GetAllocatedSize() const
+{
+	return SimCallback_AllocatedSize;
+}
+#endif
+
 void UBuoyancySubsystem::CreateSimCallback()
 {
 	// Sometimes in PIE, sim callbacks have not been freed at this point.
@@ -369,6 +395,7 @@ void UBuoyancySubsystem::ProcessSurfaceTouchCallbacks()
 
 	while (Chaos::TSimCallbackOutputHandle<FBuoyancySubsystemSimCallbackOutput> AsyncOutput = SimCallback->PopFutureOutputData_External())
 	{
+		// Process surface touches
 		for (const FBuoyancySubsystemSimCallbackOutput::FSurfaceTouch& SurfaceTouch : AsyncOutput->SurfaceTouches)
 		{
 			// Skip if we mask out this touch type
@@ -433,6 +460,14 @@ void UBuoyancySubsystem::ProcessSurfaceTouchCallbacks()
 			DispatchEvent(WaterActor);
 			DispatchEvent(RigidComponent->GetOwner());
 		}
+
+#if WITH_BUOYANCY_MEMORY_TRACKING
+		// Process changes to tracking of memory allocated on physics thread
+		if (AsyncOutput->AllocatedSize.IsSet())
+		{
+			SimCallback_AllocatedSize = *AsyncOutput->AllocatedSize;
+		}
+#endif
 	}
 }
 
@@ -482,6 +517,17 @@ void FBuoyancySubsystemSimCallbackInput::Reset()
 void FBuoyancySubsystemSimCallbackOutput::Reset()
 {
 	SurfaceTouches.Reset();
+
+#if WITH_BUOYANCY_MEMORY_TRACKING
+	AllocatedSize.Reset();
+#endif
+}
+
+SIZE_T FBuoyancySubsystemSimCallback::GetAllocatedSize() const
+{
+	return
+		BuoyancyParticleData.GetAllocatedSize() +
+		SplineKeyCache.GetAllocatedSize();
 }
 
 void FBuoyancySubsystemSimCallback::OnPreSimulate_Internal()
@@ -512,9 +558,9 @@ void FBuoyancySubsystemSimCallback::OnPreSimulate_Internal()
 			// filtering options. It might be desirable to instead add a mechanism to
 			// suppress events for one frame instead, or to remove only the data which
 			// relate to water particles which have been removed.
-			Interactions.Reset();
-			SubmersionMetaData.Reset();
-			PrevSubmersionMetaData.Reset();
+			BuoyancyParticleData.Interactions.Reset();
+			BuoyancyParticleData.SubmersionMetaData.Reset();
+			BuoyancyParticleData.PrevSubmersionMetaData.Reset();
 		}
 
 		if (Input->BuoyancySettings.IsValid())
@@ -595,14 +641,14 @@ void FBuoyancySubsystemSimCallback::OnMidPhaseModification_Internal(Chaos::FMidP
 	// SparseArray implements move semantics, so these swaps should amount to pointer swaps.
 	// This way array memories stick around even when reset/swapped so we don't do many
 	// new allocations.
-	Swap(Submersions, PrevSubmersions);
-	Swap(SubmersionMetaData, PrevSubmersionMetaData);
+	Swap(BuoyancyParticleData.Submersions, BuoyancyParticleData.PrevSubmersions);
+	Swap(BuoyancyParticleData.SubmersionMetaData, BuoyancyParticleData.PrevSubmersionMetaData);
 
 	// Clear arrays for the following phases, but keep their memory allocated
-	Interactions.Reset();
-	Submersions.Reset();
-	SubmergedShapes.Reset();
-	SubmersionMetaData.Reset();
+	BuoyancyParticleData.Interactions.Reset();
+	BuoyancyParticleData.Submersions.Reset();
+	BuoyancyParticleData.SubmergedShapes.Reset();
+	BuoyancyParticleData.SubmersionMetaData.Reset();
 
 	// Build list of "submersions"
 	TrackInteractions(*Evolution, MidPhaseAccessor);
@@ -615,6 +661,11 @@ void FBuoyancySubsystemSimCallback::OnMidPhaseModification_Internal(Chaos::FMidP
 
 	// Generate async outputs for callback data
 	GenerateCallbackData();
+
+#if WITH_BUOYANCY_MEMORY_TRACKING
+	// Generate async outputs for allocation data
+	GenerateAllocationData();
+#endif
 }
 
 void FBuoyancySubsystemSimCallback::TrackInteractions(
@@ -718,33 +769,21 @@ void FBuoyancySubsystemSimCallback::TrackInteraction(
 	{
 		SCOPE_CYCLE_COUNTER(STAT_BuoyancySubsystem_AddInteraction)
 
-		// Create the interaction struct
-		FBuoyancyInteraction BuoyancyInteraction
-		{
+		// If this particle already has a list of interactions, add to it. Otherwise
+		// make a new one.
+		BuoyancyParticleData.GetInteractions(*RigidParticle).Add({
 			RigidParticle,
 			WaterParticle,
 			WaterSpline,
 			ClosestSplineKey,
 			ClosestPoint,
-		};
-
-		// If this particle already has a list of interactions, add to it. Otherwise
-		// make a new one.
-		const int32 RigidParticleIndex = RigidParticle->UniqueIdx().Idx;
-		if (Interactions.IsValidIndex(RigidParticleIndex) == false)
-		{
-			Interactions.Insert(RigidParticleIndex, { BuoyancyInteraction });
-		}
-		else
-		{
-			Interactions[RigidParticleIndex].Add(BuoyancyInteraction);
-		}
+		});
 	}
 }
 
 void FBuoyancySubsystemSimCallback::ProcessInteractions(Chaos::FPBDRigidsEvolution& Evolution)
 {
-	for (TArray<FBuoyancyInteraction, TInlineAllocator<MaxNumBuoyancyInteractions>>& ParticleInteractions : Interactions)
+	for (FBuoyancyInteractionArray& ParticleInteractions : BuoyancyParticleData.Interactions)
 	{
 		// Sort this particle's water interactions by water level - highest to lowest
 		ParticleInteractions.Sort([](const FBuoyancyInteraction& A, const FBuoyancyInteraction& B)
@@ -853,9 +892,6 @@ void FBuoyancySubsystemSimCallback::ProcessInteraction(Chaos::FPBDRigidsEvolutio
 				PrevPoint = SplinePoint;
 			}
 
-			// Draw water velocity at the surface point
-			//Chaos::FDebugDrawQueue::GetInstance().DrawDebugDirectionalArrow(SurfacePoint, SurfacePoint + WaterVel, 20.f, FColor::Yellow, false, -1.f, -1, 3.f);
-
 			// Draw water surface normal at the surface point
 			Chaos::FDebugDrawQueue::GetInstance().DrawDebugDirectionalArrow(SurfacePoint, SurfacePoint + (WaterN * 100.f), 20.f, FColor::Green, false, -1.f, -1, 3.f);
 		}
@@ -866,7 +902,7 @@ void FBuoyancySubsystemSimCallback::ProcessInteraction(Chaos::FPBDRigidsEvolutio
 	float SubmergedVol;
 	Chaos::FVec3 SubmergedCoM;
 	float TotalVol;
-	if (BuoyancyAlgorithms::ComputeSubmergedVolume(Evolution, Interaction.RigidParticle, Interaction.WaterParticle, Interaction.ClosestPoint, WaterN, BuoyancySettings->MaxNumBoundsSubdivisions, BuoyancySettings->MinBoundsSubdivisionVol, SubmergedShapes, SubmergedVol, SubmergedCoM, TotalVol))
+	if (BuoyancyAlgorithms::ComputeSubmergedVolume(BuoyancyParticleData, Evolution, Interaction.RigidParticle, Interaction.WaterParticle, Interaction.ClosestPoint, WaterN, BuoyancySettings->MaxNumBoundsSubdivisions, BuoyancySettings->MinBoundsSubdivisionVol, SubmergedVol, SubmergedCoM, TotalVol))
 	{
 		SCOPE_CYCLE_COUNTER(STAT_BuoyancySubsystem_BuildSubmersions)
 
@@ -874,43 +910,38 @@ void FBuoyancySubsystemSimCallback::ProcessInteraction(Chaos::FPBDRigidsEvolutio
 		// approximate submerged center of mass, and apply the buoyancy force there.
 		if (SubmergedVol > SMALL_NUMBER)
 		{
-			const int32 RigidParticleIndex = Interaction.RigidParticle->UniqueIdx().Idx;
-
 			// Get the water velocity at this point
 			const FVector WaterVel
 				= Interaction.WaterSpline.Velocity.IsSet()
 				? Interaction.WaterSpline.Velocity->Eval(Interaction.ClosestSplineKey) * ClosestPosDerivative.GetSafeNormal()
 				: Chaos::FVec3::ZeroVector;
 
-			// If this particle was already marked submerged, add to its existing submersion
-			if (Submersions.IsValidIndex(RigidParticleIndex))
+			// Get the submersion for this particle. If it's new, the particle ptr will be null.
+			// If it's not new, make sure the particle ptr is the right particle.
+			FBuoyancySubmersion& Submersion = BuoyancyParticleData.GetSubmersion(*Interaction.RigidParticle);
+			if (Submersion.Particle != nullptr)
 			{
-				FBuoyancySubmersion& Submersion = Submersions[RigidParticleIndex];
 				ensureMsgf(Submersion.Particle == Interaction.RigidParticle, TEXT("Something went wrong - there's a particle index mismatch in the Submersions sparse array"));
+			}
 
-				// Get the weighted-average CoM
-				// NOTE: The unchecked division should be safe since we already
-				// know SubmergedVol > SMALL_NUMBER
-				const float TotalSubmergedVol = Submersion.Vol + SubmergedVol;
-				Submersion.CoM = ((Submersion.CoM * Submersion.Vol) + (SubmergedCoM * SubmergedVol)) / TotalSubmergedVol;
+			Submersion.Particle = Interaction.RigidParticle;
 
-				// Sum the volumes
-				Submersion.Vol = TotalSubmergedVol;
+			// Get the weighted-average CoM
+			// NOTE: The unchecked division should be safe since we already
+			// know SubmergedVol > SMALL_NUMBER
+			const float TotalSubmergedVol = Submersion.Vol + SubmergedVol;
+			Submersion.CoM = ((Submersion.CoM * Submersion.Vol) + (SubmergedCoM * SubmergedVol)) / TotalSubmergedVol;
 
-				// Get the weighted-average slerped water surface norm
-				const float VolRatio = SubmergedVol / TotalSubmergedVol;
-				const FQuat NormRot = FQuat::FindBetweenNormals(Submersion.Norm, WaterN);
-				Submersion.Norm = NormRot.Slerp(FQuat::Identity, NormRot, VolRatio).GetUpVector();
+			// Sum the volumes
+			Submersion.Vol = TotalSubmergedVol;
 
-				// Blend water velocities
-				Submersion.Vel = FMath::Lerp(Submersion.Vel, WaterVel, VolRatio);
-			}
+			// Get the weighted-average slerped water surface norm
+			const float VolRatio = SubmergedVol / TotalSubmergedVol;
+			const FQuat NormRot = FQuat::FindBetweenNormals(Submersion.Norm, WaterN);
+			Submersion.Norm = NormRot.Slerp(FQuat::Identity, NormRot, VolRatio).GetUpVector();
 
-			// If this particle was not yet submerged, make a new submersion for it
-			else
-			{
-				Submersions.Insert(RigidParticleIndex, { Interaction.RigidParticle, SubmergedVol, SubmergedCoM, WaterVel, WaterN });
-			}
+			// Blend water velocities
+			Submersion.Vel = FMath::Lerp(Submersion.Vel, WaterVel, VolRatio);
 
 			// If this is a surface touch record it for callback, if 
 			if (BuoyancySettings->SurfaceTouchCallbackFlags != 0 &&
@@ -919,7 +950,6 @@ void FBuoyancySubsystemSimCallback::ProcessInteraction(Chaos::FPBDRigidsEvolutio
 				SCOPE_CYCLE_COUNTER(STAT_BuoyancySubsystem_BuildSubmersionCallbackData)
 
 				// Proceed only if this CoM is moving fast enough to generate events
-				const FBuoyancySubmersion& Submersion = Submersions[RigidParticleIndex];
 				const Chaos::FVec3 CoMDiff = Submersion.CoM - Interaction.RigidParticle->XCom();
 				const Chaos::FVec3 CoMVel = Interaction.RigidParticle->GetV() + Chaos::FVec3::CrossProduct(Interaction.RigidParticle->GetW(), CoMDiff);
 				const float CoMVelSq = Chaos::FVec3::DotProduct(CoMVel, CoMVel);
@@ -929,11 +959,7 @@ void FBuoyancySubsystemSimCallback::ProcessInteraction(Chaos::FPBDRigidsEvolutio
 				if (CoMVelSq > MinVelSq)
 				{
 					// If we don't have a metadata for this particle yet, add one
-					if (!SubmersionMetaData.IsValidIndex(RigidParticleIndex))
-					{
-						SubmersionMetaData.Insert(RigidParticleIndex, FBuoyancySubmersionMetaData());
-					}
-					FBuoyancySubmersionMetaData& MetaData = SubmersionMetaData[RigidParticleIndex];
+					FBuoyancySubmersionMetaData& MetaData = BuoyancyParticleData.GetSubmersionMetaData(*Interaction.RigidParticle);
 
 					// If we haven't already maxed out on water contacts, add one
 					if (MetaData.WaterContacts.Num() < FBuoyancySubmersionMetaData::MaxNumWaterContacts)
@@ -957,7 +983,7 @@ void FBuoyancySubsystemSimCallback::ApplyBuoyantForces(Chaos::FPBDRigidsEvolutio
 	const Chaos::FPerParticleGravity* PerParticleGravity = &Evolution.GetGravityForces();
 
 	// Apply all buoyant forces
-	for (const FBuoyancySubmersion& Submersion : Submersions)
+	for (const FBuoyancySubmersion& Submersion : BuoyancyParticleData.Submersions)
 	{
 		// Figure out the gravity level of the particle
 		const int32 GravityGroupIndex = Submersion.Particle->GravityGroupIndex();
@@ -1024,6 +1050,11 @@ void FBuoyancySubsystemSimCallback::GenerateCallbackData()
 		return;
 	}
 
+	// Get refs to internal sparse arrays of submersion data
+	TSparseArray<FBuoyancySubmersion>& PrevSubmersions = BuoyancyParticleData.PrevSubmersions;
+	TSparseArray<FBuoyancySubmersionMetaData>& MetaDatas = BuoyancyParticleData.SubmersionMetaData;
+	TSparseArray<FBuoyancySubmersionMetaData>& PrevMetaDatas = BuoyancyParticleData.PrevSubmersionMetaData;
+
 	// Generate callback data if we're into that sort of thing
 	const uint8 CallbackFlags = BuoyancySettings->SurfaceTouchCallbackFlags;
 	if (CallbackFlags != 0)
@@ -1035,18 +1066,18 @@ void FBuoyancySubsystemSimCallback::GenerateCallbackData()
 
 		// Process every surface touch and queue up some of them to return
 		// to game thread for callback dispatch
-		Output.SurfaceTouches.Reserve(SubmersionMetaData.Num() * FBuoyancySubmersionMetaData::MaxNumWaterContacts);
-		for (auto Iter = SubmersionMetaData.CreateIterator(); Iter; ++Iter)
+		Output.SurfaceTouches.Reserve(MetaDatas.Num() * FBuoyancySubmersionMetaData::MaxNumWaterContacts);
+		for (auto Iter = MetaDatas.CreateIterator(); Iter; ++Iter)
 		{
 			const FBuoyancySubmersionMetaData& MetaData = *Iter;
 			const int32 ObjectIndex = Iter.GetIndex();
-			const FBuoyancySubmersion& Submersion = Submersions[ObjectIndex];
+			const FBuoyancySubmersion& Submersion = BuoyancyParticleData.Submersions[ObjectIndex];
 
 			// Mark this as a new or continuing contact based on whether or
 			// not we have a bit from the previous-submersions array.
 			const bool bPrevSubmerged =
-				PrevSubmersionMetaData.IsValidIndex(ObjectIndex) &&
-				PrevSubmersionMetaData.IsAllocated(ObjectIndex);
+				PrevMetaDatas.IsValidIndex(ObjectIndex) &&
+				PrevMetaDatas.IsAllocated(ObjectIndex);
 			const EBuoyancyEventFlags TouchFlag
 				= bPrevSubmerged
 				? EBuoyancyEventFlags::Continue
@@ -1057,7 +1088,7 @@ void FBuoyancySubsystemSimCallback::GenerateCallbackData()
 			// bother doing this work if we're tracking removals
 			if (bPrevSubmerged && (CallbackFlags & EBuoyancyEventFlags::End) != 0)
 			{
-				PrevSubmersionMetaData.RemoveAt(ObjectIndex);
+				PrevMetaDatas.RemoveAt(ObjectIndex);
 			}
 
 			// Only continue if we're tracking this touch type
@@ -1097,7 +1128,7 @@ void FBuoyancySubsystemSimCallback::GenerateCallbackData()
 		// loses contact, not just when one part of it loses contact.
 		if ((CallbackFlags & EBuoyancyEventFlags::End) != 0)
 		{
-			for (auto Iter = PrevSubmersionMetaData.CreateIterator(); Iter; ++Iter)
+			for (auto Iter = PrevMetaDatas.CreateIterator(); Iter; ++Iter)
 			{
 				const FBuoyancySubmersionMetaData& MetaData = *Iter;
 				const int32 ObjectIndex = Iter.GetIndex();
@@ -1128,3 +1159,27 @@ void FBuoyancySubsystemSimCallback::GenerateCallbackData()
 		}
 	}
 }
+
+#if WITH_BUOYANCY_MEMORY_TRACKING
+void FBuoyancySubsystemSimCallback::GenerateAllocationData()
+{
+	SCOPE_CYCLE_COUNTER(STAT_BuoyancySubsystem_GenerateAllocationData)
+
+	// NOTE: For now I'm guarding this inside WITH_BUOYANCY_MEMORY_TRACKING.
+	// The current use case for the game thread GetAllocatedSize is a console
+	// command which is editor-only. If this is to become useful
+	// more generally, we can find a way to do it without relying on
+	// continual calls to GetAllocatedSize().
+	const uint32 NewAllocatedSize = GetAllocatedSize();
+	if (AllocatedSize != NewAllocatedSize)
+	{
+		// If the number of allocated bytes has changed, update our cached num bytes
+		AllocatedSize = NewAllocatedSize;
+
+		// Get the async output struct to write to
+		FBuoyancySubsystemSimCallbackOutput& Output = GetProducerOutputData_Internal();
+		Output.AllocatedSize = AllocatedSize;
+	}
+}
+#endif
+
diff --git a/Engine/Plugins/Experimental/Buoyancy/Source/Runtime/Private/BuoyancyWaterSplineKeyCacheGrid.cpp b/Engine/Plugins/Experimental/Buoyancy/Source/Runtime/Private/BuoyancyWaterSplineKeyCacheGrid.cpp
--- a/Engine/Plugins/Experimental/Buoyancy/Source/Runtime/Private/BuoyancyWaterSplineKeyCacheGrid.cpp
+++ b/Engine/Plugins/Experimental/Buoyancy/Source/Runtime/Private/BuoyancyWaterSplineKeyCacheGrid.cpp
@@ -2,9 +2,12 @@
 
 #include "BuoyancyWaterSplineKeyCacheGrid.h"
 #include "CoreMinimal.h"
+#include "BuoyancyStats.h"
 
 float FSplineKeyCacheGrid::GetClosestSplineKey(const FBuoyancyWaterSplineData& WaterSplineData, const Chaos::FVec3& WorldPosition)
 {
+	LLM_SCOPE_BYTAG(BuoyancyParticleDataTag);
+
 	// Convert the world position into a position local to the water spline
 	const FVector LocalPosition = WaterSplineData.Transform.InverseTransformPosition(WorldPosition);
 
@@ -84,6 +87,11 @@ uint32 FSplineKeyCacheGrid::GetCacheLimit() const
 	return CacheLimit;
 }
 
+SIZE_T FSplineKeyCacheGrid::GetAllocatedSize() const
+{
+	return SplineKeyCache.GetAllocatedSize();
+}
+
 Chaos::TVec3<int32> FSplineKeyCacheGrid::GetCacheKey(const FVector& LocalPos) const
 {
 	const float GridSizeInv = 1.f / GridSize;
diff --git a/Engine/Plugins/Experimental/Buoyancy/Source/Runtime/Public/BuoyancyAlgorithms.h b/Engine/Plugins/Experimental/Buoyancy/Source/Runtime/Public/BuoyancyAlgorithms.h
--- a/Engine/Plugins/Experimental/Buoyancy/Source/Runtime/Public/BuoyancyAlgorithms.h
+++ b/Engine/Plugins/Experimental/Buoyancy/Source/Runtime/Public/BuoyancyAlgorithms.h
@@ -12,6 +12,8 @@ namespace Chaos
 	class FPBDRigidsEvolutionGBF;
 }
 
+struct FBuoyancyParticleData;
+
 namespace BuoyancyAlgorithms
 {
 	// Minimal struct containing essential data about a particular submersion
@@ -40,15 +42,15 @@ namespace BuoyancyAlgorithms
 
 	// Compute an approximate volume and center of mass of particle B submerged in particle A,
 	// adjusting for the volume of the object based on the material density and mass of the object
-	bool ComputeSubmergedVolume(const Chaos::FPBDRigidsEvolutionGBF& Evolution, const Chaos::FGeometryParticleHandle* ParticleA, const Chaos::FGeometryParticleHandle* ParticleB, int32 NumSubdivisions, float MinVolume, TSparseArray<TBitArray<>>& SubmergedShapes, float& SubmergedVol, Chaos::FVec3& SubmergedCoM, float& TotalVol);
+	bool ComputeSubmergedVolume(FBuoyancyParticleData& ParticleData, const Chaos::FPBDRigidsEvolutionGBF& Evolution, const Chaos::FGeometryParticleHandle* ParticleA, const Chaos::FGeometryParticleHandle* ParticleB, int32 NumSubdivisions, float MinVolume, float& SubmergedVol, Chaos::FVec3& SubmergedCoM, float& TotalVol);
 
 	// Compute an approximate volume and center of mass of particle B submerged in particle A
-	bool ComputeSubmergedVolume(const Chaos::FGeometryParticleHandle* ParticleA, const Chaos::FGeometryParticleHandle* ParticleB, int32 NumSubdivisions, float MinVolume, TSparseArray<TBitArray<>>& SubmergedShapes, float& SubmergedVol, Chaos::FVec3& SubmergedCoM);
+	bool ComputeSubmergedVolume(FBuoyancyParticleData& ParticleData, const Chaos::FGeometryParticleHandle* ParticleA, const Chaos::FGeometryParticleHandle* ParticleB, int32 NumSubdivisions, float MinVolume, float& SubmergedVol, Chaos::FVec3& SubmergedCoM);
 
 	// Compute submerged volume given a single waterlevel
-	bool ComputeSubmergedVolume(const Chaos::FPBDRigidsEvolutionGBF& Evolution, const Chaos::FGeometryParticleHandle* SubmergedParticle, const Chaos::FGeometryParticleHandle* WaterParticle, const FVector& WaterX, const FVector& WaterN, int32 NumSubdivisions, float MinVolume, TSparseArray<TBitArray<>>& SubmergedShapes, float& SubmergedVol, Chaos::FVec3& SubmergedCoM, float& TotalVol);
+	bool ComputeSubmergedVolume(FBuoyancyParticleData& ParticleData, const Chaos::FPBDRigidsEvolutionGBF& Evolution, const Chaos::FGeometryParticleHandle* SubmergedParticle, const Chaos::FGeometryParticleHandle* WaterParticle, const FVector& WaterX, const FVector& WaterN, int32 NumSubdivisions, float MinVolume, float& SubmergedVol, Chaos::FVec3& SubmergedCoM, float& TotalVol);
 
-	bool ComputeSubmergedVolume(const Chaos::FGeometryParticleHandle* SubmergedParticle, const Chaos::FGeometryParticleHandle* WaterParticle, const FVector& WaterX, const FVector& WaterN, int32 NumSubdivisions, float MinVolume, TSparseArray<TBitArray<>>& SubmergedShapes, float& SubmergedVol, Chaos::FVec3& SubmergedCoM);
+	bool ComputeSubmergedVolume(FBuoyancyParticleData& ParticleData, const Chaos::FGeometryParticleHandle* SubmergedParticle, const Chaos::FGeometryParticleHandle* WaterParticle, const FVector& WaterX, const FVector& WaterN, int32 NumSubdivisions, float MinVolume, float& SubmergedVol, Chaos::FVec3& SubmergedCoM);
 
 	// Given an OOBB and a water level, generate another OOBB which is 1. entirely contained
 	// within the input OOBB and 2. entirely contains the portion of the OOBB which is submerged
diff --git a/Engine/Plugins/Experimental/Buoyancy/Source/Runtime/Public/BuoyancyMacros.h b/Engine/Plugins/Experimental/Buoyancy/Source/Runtime/Public/BuoyancyMacros.h
new file mode 100644
--- /dev/null
+++ b/Engine/Plugins/Experimental/Buoyancy/Source/Runtime/Public/BuoyancyMacros.h
@@ -0,0 +1,7 @@
+// Copyright Epic Games, Inc. All Rights Reserved.
+
+#pragma once
+
+#include "CoreMinimal.h"
+
+#define WITH_BUOYANCY_MEMORY_TRACKING (!UE_BUILD_SHIPPING && !UE_BUILD_TEST)
diff --git a/Engine/Plugins/Experimental/Buoyancy/Source/Runtime/Public/BuoyancyParticleData.h b/Engine/Plugins/Experimental/Buoyancy/Source/Runtime/Public/BuoyancyParticleData.h
new file mode 100644
--- /dev/null
+++ b/Engine/Plugins/Experimental/Buoyancy/Source/Runtime/Public/BuoyancyParticleData.h
@@ -0,0 +1,178 @@
+// Copyright Epic Games, Inc. All Rights Reserved.
+
+#pragma once
+
+//#include "Containers/HashTable.h"
+#include "CoreMinimal.h"
+#include "Chaos/SimCallbackObject.h"
+#include "Chaos/Framework/PhysicsProxyBase.h"
+#include "Chaos/PBDRigidsEvolutionFwd.h"
+#include "PBDRigidsSolver.h"
+#include "Engine/EngineBaseTypes.h"
+#include "PhysicsProxy/SingleParticlePhysicsProxyFwd.h"
+#include "BuoyancyWaterSplineData.h"
+#include "BuoyancyStats.h"
+
+
+// Each particle will have a list of potential midphases to process,
+// which must be sorted in descending Z order. This struct is used
+// to store them
+struct FBuoyancyInteraction
+{
+	Chaos::FPBDRigidParticleHandle* RigidParticle = nullptr;
+	Chaos::FGeometryParticleHandle* WaterParticle = nullptr;
+	const FBuoyancyWaterSplineData& WaterSpline;
+	float ClosestSplineKey = 0.f;
+	FVector ClosestPoint = FVector::ZeroVector;
+};
+
+static constexpr int32 MaxNumBuoyancyInteractions = 2;
+using FBuoyancyInteractionArray = TArray<FBuoyancyInteraction, TInlineAllocator<MaxNumBuoyancyInteractions>>;
+
+// A minimal struct of data tracking all the submersions in a frame.
+struct FBuoyancySubmersion
+{
+	Chaos::FPBDRigidParticleHandle* Particle = nullptr;
+	float Vol = 0.f;
+	FVector CoM = FVector::ZeroVector;
+	FVector Vel = FVector::ZeroVector;
+	FVector Norm = FVector::ZeroVector;
+};
+
+// Metadata for submersions, used for event callbacks
+struct FBuoyancySubmersionMetaData
+{
+	struct FWaterContact
+	{
+		Chaos::FGeometryParticleHandle* Water = nullptr;
+		float Vol = 0.f;
+		FVector CoM = FVector::ZeroVector;
+		FVector Vel = FVector::ZeroVector;
+	};
+
+	// How many metadata's allowed per submerged particle
+	static constexpr int32 MaxNumWaterContacts = 3;
+	TArray<FWaterContact, TInlineAllocator<MaxNumWaterContacts>> WaterContacts;
+};
+
+struct FBuoyancyParticleData
+{
+private:
+
+	// Access an element in a specific array, adding an uninitialized one if
+	// it doesn't exist yet.
+	template <typename TData>
+	TData& GetData(const Chaos::FGeometryParticleHandle& ParticleHandle, TSparseArray<TData>& DataArray)
+	{
+		SCOPE_CYCLE_COUNTER(STAT_BuoyancyParticleData_GetData)
+
+		const int32 DataIndex = AddOrGetIndex(ParticleHandle);
+
+		if (DataArray.IsValidIndex(DataIndex) == false)
+		{
+			LLM_SCOPE_BYTAG(BuoyancyParticleDataTag);
+
+			DataArray.Insert(DataIndex, TData());
+		}
+
+		return DataArray[DataIndex];
+	}
+
+public:
+
+	FBuoyancyParticleData();
+	~FBuoyancyParticleData();
+
+	// Clear all data without freeing memory, for quick repopulation
+	void Reset();
+
+	// For memory debugging, log the number of bytes that this class has allocated
+	SIZE_T GetAllocatedSize() const;
+
+	// Internal index accessor
+	int32 GetIndex(const Chaos::FGeometryParticleHandle& ParticleHandle);
+
+	int32 AddOrGetIndex(const Chaos::FGeometryParticleHandle& ParticleHandle);
+
+	bool RemoveIndex(const Chaos::FGeometryParticleHandle& ParticleHandle);
+
+	// Bookkeeping arrays
+	//
+	// Map of indices from unique particle indices to internal array indices
+	//TSparseArray<int32> IndexMap;
+	FHashTable IndexMap;
+	TSparseArray<int32> ReverseIndexMap;
+
+	// Sparse array of arrays of buoyancy interactions - the outer array has one
+	// entry per particle, the inner array has an entry per water body that it interacts
+	// with. Each will be a very small array, sorted by Z.
+	//
+	// We use inline allocator to avoid more heap allocations, and to express the
+	// assumption that a single particle is unlikely to exceed interactions with a
+	// certain number of waterbodies at a time.
+	TSparseArray<FBuoyancyInteractionArray> Interactions;
+
+	// This sparse array of submersion events is indexed on particle unique indices.
+	// All buoyant forces due to submersions are applied at once. It's stored as
+	// a member variable and reset every frame, to avoid reallocation of similarly
+	// sized data.
+	TSparseArray<FBuoyancySubmersion> Submersions;
+	TSparseArray<FBuoyancySubmersion> PrevSubmersions;
+
+	// Another sparse array to be kept in sync with Submersions, which will contain
+	// metadata useful for event callbacks
+	TSparseArray<FBuoyancySubmersionMetaData> SubmersionMetaData;
+	TSparseArray<FBuoyancySubmersionMetaData> PrevSubmersionMetaData;
+
+	// This is a sparse array of bit arrays representing which shapes in an object
+	// have already been accounted for when submerging an object. For example, if
+	// a massive BVH object has two leaf node shapes submerged in different pools
+	// of water and we've already detected that leaf A is submerged, we don't
+	// need to test A again. This helps to avoid double counting submerged shapes.
+	//
+	// Just like Submersions, we have this as a member variable only to keep the
+	// memory hot - the array is reset, repopulated and traversed, every frame,
+	// so we want to minimize allocations.
+	TSparseArray<TBitArray<>> SubmergedShapes;
+
+	// Element accessors
+	FBuoyancyInteractionArray& GetInteractions(const Chaos::FGeometryParticleHandle& ParticleHandle)
+	{
+		return GetData(ParticleHandle, Interactions);
+	}
+
+	FBuoyancySubmersion& GetSubmersion(const Chaos::FGeometryParticleHandle& ParticleHandle)
+	{
+		return GetData(ParticleHandle, Submersions);
+	}
+
+	FBuoyancySubmersion& GetPrevSubmersion(const Chaos::FGeometryParticleHandle& ParticleHandle)
+	{
+		return GetData(ParticleHandle, PrevSubmersions);
+	}
+
+	FBuoyancySubmersionMetaData& GetSubmersionMetaData(const Chaos::FGeometryParticleHandle& ParticleHandle)
+	{
+		return GetData(ParticleHandle, SubmersionMetaData);
+	}
+
+	FBuoyancySubmersionMetaData& GetPrevSubmersionMetaData(const Chaos::FGeometryParticleHandle& ParticleHandle)
+	{
+		return GetData(ParticleHandle, PrevSubmersionMetaData);
+	}
+
+	TBitArray<>& GetSubmergedShapes(const Chaos::FGeometryParticleHandle& ParticleHandle)
+	{
+		return GetData(ParticleHandle, SubmergedShapes);
+	}
+
+private:
+
+	int32 GetIndex(const Chaos::FGeometryParticleHandle& ParticleHandle, int32& OutParticleIndex, int32& OutParticleKey);
+
+	int32 GetIndex(const int32 ParticleIndex, const int32 ParticleKey);
+
+	/** Shrink or grow internal memory - this operation may be slow when a resize is performed,
+	    but should result in more optimal storage and faster data accesses on average. */
+	void OptimizeMemory();
+};
diff --git a/Engine/Plugins/Experimental/Buoyancy/Source/Runtime/Public/BuoyancyStats.h b/Engine/Plugins/Experimental/Buoyancy/Source/Runtime/Public/BuoyancyStats.h
--- a/Engine/Plugins/Experimental/Buoyancy/Source/Runtime/Public/BuoyancyStats.h
+++ b/Engine/Plugins/Experimental/Buoyancy/Source/Runtime/Public/BuoyancyStats.h
@@ -20,6 +20,7 @@ DECLARE_CYCLE_STAT(TEXT("Buoyancy_Subsystem_BuildSubmersions"), STAT_BuoyancySub
 DECLARE_CYCLE_STAT(TEXT("Buoyancy_Subsystem_BuildSubmersionCallbackData"), STAT_BuoyancySubsystem_BuildSubmersionCallbackData, STATGROUP_Buoyancy);
 DECLARE_CYCLE_STAT(TEXT("Buoyancy_Subsystem_ApplyBuoyantForces"), STAT_BuoyancySubsystem_ApplyBuoyantForces, STATGROUP_Buoyancy);
 DECLARE_CYCLE_STAT(TEXT("Buoyancy_Subsystem_ProduceSurfaceTouches"), STAT_BuoyancySubsystem_ProduceSurfaceTouches, STATGROUP_Buoyancy);
+DECLARE_CYCLE_STAT(TEXT("Buoyancy_Subsystem_GenerateAllocationData"), STAT_BuoyancySubsystem_GenerateAllocationData, STATGROUP_Buoyancy);
 DECLARE_CYCLE_STAT(TEXT("Buoyancy_Subsystem_Tick"), STAT_BuoyancySubsystem_Tick, STATGROUP_Buoyancy);
 DECLARE_CYCLE_STAT(TEXT("Buoyancy_Subsystem_UpdateBuoyancySettings"), STAT_BuoyancySubsystem_UpdateBuoyancySettings, STATGROUP_Buoyancy);
 DECLARE_CYCLE_STAT(TEXT("Buoyancy_Subsystem_UpdateWaterBodiesList"), STAT_BuoyancySubsystem_UpdateWaterBodiesList, STATGROUP_Buoyancy);
@@ -33,3 +34,10 @@ DECLARE_CYCLE_STAT(TEXT("Buoyancy_Algorithms_SubdivideBounds"), STAT_BuoyancyAlg
 DECLARE_CYCLE_STAT(TEXT("Buoyancy_Algorithms_ComputeBuoyantForces"), STAT_BuoyancyAlgorithms_ComputeBuoyantForces, STATGROUP_Buoyancy);
 DECLARE_CYCLE_STAT(TEXT("Buoyancy_Algorithms_SubmergeShapeInternal"), STAT_BuoyancyAlgorithms_SubmergeShapeInternal, STATGROUP_Buoyancy);
 
+// Particle Data
+DECLARE_CYCLE_STAT(TEXT("Buoyancy_ParticleData_GetData"), STAT_BuoyancyParticleData_GetData, STATGROUP_Buoyancy);
+DECLARE_CYCLE_STAT(TEXT("Buoyancy_ParticleData_Reset"), STAT_BuoyancyParticleData_Reset, STATGROUP_Buoyancy);
+DECLARE_CYCLE_STAT(TEXT("Buoyancy_ParticleData_OptimizeMemory"), STAT_BuoyancyParticleData_OptimizeMemory, STATGROUP_Buoyancy);
+
+// Memory tracking tags
+LLM_DECLARE_TAG(BuoyancyParticleDataTag);
diff --git a/Engine/Plugins/Experimental/Buoyancy/Source/Runtime/Public/BuoyancySubsystem.h b/Engine/Plugins/Experimental/Buoyancy/Source/Runtime/Public/BuoyancySubsystem.h
--- a/Engine/Plugins/Experimental/Buoyancy/Source/Runtime/Public/BuoyancySubsystem.h
+++ b/Engine/Plugins/Experimental/Buoyancy/Source/Runtime/Public/BuoyancySubsystem.h
@@ -14,6 +14,8 @@
 #include "BuoyancyEventFlags.h"
 #include "ChaosUserDataPT.h"
 #include "PhysicsProxy/SingleParticlePhysicsProxyFwd.h"
+#include "BuoyancyParticleData.h"
+#include "BuoyancyMacros.h"
 #include "BuoyancySubsystem.generated.h"
 
 DECLARE_LOG_CATEGORY_EXTERN(LogBuoyancySubsystem, Log, All);
@@ -103,6 +105,10 @@ public:
 
 	bool SetEnabledWithUpdatedNetModeCallback(const bool bEnabled);
 
+#if WITH_BUOYANCY_MEMORY_TRACKING
+	SIZE_T GetAllocatedSize() const;
+#endif
+
 protected:
 
 	// UTickableWorldSubsystem begin interface
@@ -153,6 +159,10 @@ private:
 	class FBuoyancySubsystemSimCallback* SimCallback;
 
 	ENetMode NetMode;
+
+#if WITH_BUOYANCY_MEMORY_TRACKING
+	uint32 SimCallback_AllocatedSize;
+#endif
 };
 
 
@@ -160,45 +170,6 @@ private:
 // Buoyancy Sim Callback
 //
 
-// Each particle will have a list of potential midphases to process,
-// which must be sorted in descending Z order. This struct is used
-// to store them
-struct FBuoyancyInteraction
-{
-	Chaos::FPBDRigidParticleHandle* RigidParticle;
-	Chaos::FGeometryParticleHandle* WaterParticle;
-	const FBuoyancyWaterSplineData& WaterSpline;
-	float ClosestSplineKey;
-	FVector ClosestPoint;
-};
-
-// Metadata for submersions, used for event callbacks
-struct FBuoyancySubmersionMetaData
-{
-	struct FWaterContact
-	{
-		Chaos::FGeometryParticleHandle* Water;
-		float Vol;
-		FVector CoM;
-		FVector Vel;
-	};
-
-	// How many metadata's allowed per submerged particle
-	static constexpr int32 MaxNumWaterContacts = 3;
-	TArray<FWaterContact, TInlineAllocator<MaxNumWaterContacts>> WaterContacts;
-};
-
-
-// A minimal struct of data tracking all the submersions in a frame.
-struct FBuoyancySubmersion
-{
-	Chaos::FPBDRigidParticleHandle* Particle;
-	float Vol;
-	FVector CoM;
-	FVector Vel;
-	FVector Norm;
-};
-
 struct FBuoyancySubsystemSimCallbackInput : public Chaos::FSimCallbackInput
 {
 	// If this ptr is set, then we have a new spline data manager...
@@ -231,6 +202,11 @@ struct FBuoyancySubsystemSimCallbackOutput : public Chaos::FSimCallbackOutput
 
 	TArray<FSurfaceTouch> SurfaceTouches;
 
+#if WITH_BUOYANCY_MEMORY_TRACKING
+	// Optional update to allocated num bytes
+	TOptional<uint32> AllocatedSize;
+#endif
+
 	void Reset();
 };
 
@@ -241,6 +217,10 @@ class FBuoyancySubsystemSimCallback : public Chaos::TSimCallbackObject<
 	FBuoyancySubsystemSimCallbackOutput,
 	Chaos::ESimCallbackOptions::Presimulate | Chaos::ESimCallbackOptions::MidPhaseModification>
 {
+public:
+
+	SIZE_T GetAllocatedSize() const;
+
 private:
 
 	virtual void OnPreSimulate_Internal() override;
@@ -253,6 +233,11 @@ private:
 	void ApplyBuoyantForces(Chaos::FPBDRigidsEvolution& Evolution);
 	void GenerateCallbackData();
 
+#if WITH_BUOYANCY_MEMORY_TRACKING
+	void GenerateAllocationData();
+	SIZE_T AllocatedSize = 0;
+#endif
+
 	// Reference to UserDataPT sim callback which manages synchronization of
 	// water spline data
 	FBuoyancyWaterSplineDataManager* SplineData = nullptr;
@@ -262,38 +247,7 @@ private:
 	// memory that was allocated by GT to minimize copies.
 	TUniquePtr<FBuoyancySettings> BuoyancySettings;
 
-	// Sparse array of arrays of buoyancy interactions - the outer array has one
-	// entry per particle, the inner array has an entry per water body that it interacts
-	// with. Each will be a very small array, sorted by Z.
-	//
-	// We use inline allocator to avoid more heap allocations, and to express the
-	// assumption that a single particle is unlikely to exceed interactions with a
-	// certain number of waterbodies at a time.
-	static constexpr int32 MaxNumBuoyancyInteractions = 2;
-	TSparseArray<TArray<FBuoyancyInteraction, TInlineAllocator<MaxNumBuoyancyInteractions>>> Interactions;
-
-	// This sparse array of submersion events is indexed on particle unique indices.
-	// All buoyant forces due to submersions are applied at once. It's stored as
-	// a member variable and reset every frame, to avoid reallocation of similarly
-	// sized data.
-	TSparseArray<FBuoyancySubmersion> Submersions;
-	TSparseArray<FBuoyancySubmersion> PrevSubmersions;
-
-	// Another sparse array to be kept in sync with Submersions, which will contain
-	// metadata useful for event callbacks
-	TSparseArray<FBuoyancySubmersionMetaData> SubmersionMetaData;
-	TSparseArray<FBuoyancySubmersionMetaData> PrevSubmersionMetaData;
-
-	// This is a sparse array of bit arrays representing which shapes in an object
-	// have already been accounted for when submerging an object. For example, if
-	// a massive BVH object has two leaf node shapes submerged in different pools
-	// of water and we've already detected that leaf A is submerged, we don't
-	// need to test A again. This helps to avoid double counting submerged shapes.
-	//
-	// Just like Submersions, we have this as a member variable only to keep the
-	// memory hot - the array is reset, repopulated and traversed, every frame,
-	// so we want to minimize allocations.
-	TSparseArray<TBitArray<>> SubmergedShapes;
+	FBuoyancyParticleData BuoyancyParticleData;
 
 	// Used to track the net mode of the world that owns the phys scene that this
 	// sim tick is taking place in.
diff --git a/Engine/Plugins/Experimental/Buoyancy/Source/Runtime/Public/BuoyancyWaterSplineKeyCacheGrid.h b/Engine/Plugins/Experimental/Buoyancy/Source/Runtime/Public/BuoyancyWaterSplineKeyCacheGrid.h
--- a/Engine/Plugins/Experimental/Buoyancy/Source/Runtime/Public/BuoyancyWaterSplineKeyCacheGrid.h
+++ b/Engine/Plugins/Experimental/Buoyancy/Source/Runtime/Public/BuoyancyWaterSplineKeyCacheGrid.h
@@ -21,6 +21,8 @@ public:
 	float GetGridSize() const;
 	uint32 GetCacheLimit() const;
 
+	SIZE_T GetAllocatedSize() const;
+
 private:
 
 	Chaos::TVec3<int32> GetCacheKey(const FVector& LocalPos) const;
@@ -31,5 +33,6 @@ private:
 
 	uint32 CacheLimit = 256;
 
-	TMap<const FBuoyancyWaterSplineData*, TMap<Chaos::TVec3<int32>, float>> SplineKeyCache;
+	using FSplineKeyCache = TMap<const FBuoyancyWaterSplineData*, TMap<Chaos::TVec3<int32>, float>>;
+	FSplineKeyCache SplineKeyCache;
 };
diff --git a/Engine/Source/Runtime/Core/Private/Containers/HashTable.cpp b/Engine/Source/Runtime/Core/Private/Containers/HashTable.cpp
--- a/Engine/Source/Runtime/Core/Private/Containers/HashTable.cpp
+++ b/Engine/Source/Runtime/Core/Private/Containers/HashTable.cpp
@@ -36,6 +36,14 @@ CORE_API void FHashTable::Resize( uint32 NewIndexSize )
 	NextIndex = NewNextIndex;
 }
 
+SIZE_T FHashTable::GetAllocatedSize() const
+{
+	return
+		sizeof(EmptyHash) +
+		(sizeof(uint32*) * 2) +
+		(sizeof(uint32) * (3 + IndexSize + HashSize));
+}
+
 CORE_API float FHashTable::AverageSearch() const
 {
 	uint32 SumAvgSearch = 0;
diff --git a/Engine/Source/Runtime/Core/Public/Containers/HashTable.h b/Engine/Source/Runtime/Core/Public/Containers/HashTable.h
--- a/Engine/Source/Runtime/Core/Public/Containers/HashTable.h
+++ b/Engine/Source/Runtime/Core/Public/Containers/HashTable.h
@@ -224,8 +224,9 @@ public:
 	 * If the previous size was empty, allocates the hash at its desired size.
 	 */
 	CORE_API void	Resize( uint32 NewIndexSize );
-	inline uint32	GetIndexSize() { return IndexSize; }
-	inline uint32	GetHashSize() { return HashSize; }
+	inline uint32	GetIndexSize() const { return IndexSize; }
+	inline uint32	GetHashSize() const { return HashSize; }
+	CORE_API SIZE_T	GetAllocatedSize() const;
 
 	// Functions used to search
 	uint32			First( uint32 Key ) const;
