Buoyancy subsystem only loop over midphases that involve water

diff --git a/Engine/Plugins/Experimental/Buoyancy/Source/Runtime/Private/BuoyancySubsystem.cpp b/Engine/Plugins/Experimental/Buoyancy/Source/Runtime/Private/BuoyancySubsystem.cpp
--- a/Engine/Plugins/Experimental/Buoyancy/Source/Runtime/Private/BuoyancySubsystem.cpp
+++ b/Engine/Plugins/Experimental/Buoyancy/Source/Runtime/Private/BuoyancySubsystem.cpp
@@ -23,6 +23,7 @@
 #include "Chaos/DebugDrawQueue.h"
 #include "Templates/SharedPointer.h"
 #include "PhysicsProxy/SingleParticlePhysicsProxy.h"
+#include "Chaos/GeometryParticlesFwd.h"
 
 //
 // CVars
@@ -312,6 +313,11 @@ void UBuoyancySubsystem::UpdateSplineData()
 {
 	SCOPE_CYCLE_COUNTER(STAT_BuoyancySubsystem_UpdateWaterBodiesList)
 
+	if (bWaterObjectsChanged == false)
+	{
+		return;
+	}
+
 	if (SplineData == nullptr)
 	{
 		return;
@@ -319,8 +325,15 @@ void UBuoyancySubsystem::UpdateSplineData()
 
 	if (FWaterBodyManager* WaterBodyManager = UWaterSubsystem::GetWaterBodyManager(GetWorld()))
 	{
+		// Clear the flag to rebuild water object data
 		bWaterObjectsChanged = false;
 
+		// Clear out existing old water object data
+		//
+		// NOTE: the "false" argument here prevents a resize, since we're just going to
+		// repopulate the underlying sparse array again immediately anyway.
+		SplineData->ClearData_GT(false);
+
 		// Loop over every registered water body
 		WaterBodyManager->ForEachWaterBodyComponent(GetWorld(), [this](UWaterBodyComponent* WaterBodyComponent)
 		{
@@ -624,15 +637,18 @@ void FBuoyancySubsystemSimCallback::OnMidPhaseModification_Internal(Chaos::FMidP
 		return;
 	}
 
-	// Get the evolution
-	Chaos::FPBDRigidsEvolution* Evolution = nullptr;
-	if (Chaos::FPhysicsSolverBase* SolverBase = GetSolver())
+	// Get the solver
+	Chaos::FPhysicsSolverBase* SolverBase = GetSolver();
+	if (SolverBase == nullptr)
 	{
-		// Why does cast-checked return a ref? That makes me think
-		// it's not actually doing a check...
-		Chaos::FPBDRigidsSolver& PBDSolver = SolverBase->CastChecked();
-		Evolution = PBDSolver.GetEvolution();
+		return;
 	}
+	// Why does cast-checked return a ref? That makes me think
+	// it's not actually doing a check...
+	Chaos::FPBDRigidsSolver& PBDSolver = SolverBase->CastChecked();
+
+	// Get the evolution
+	Chaos::FPBDRigidsEvolution* Evolution = PBDSolver.GetEvolution();
 	if (Evolution == nullptr)
 	{
 		return;
@@ -651,7 +667,7 @@ void FBuoyancySubsystemSimCallback::OnMidPhaseModification_Internal(Chaos::FMidP
 	BuoyancyParticleData.SubmersionMetaData.Reset();
 
 	// Build list of "submersions"
-	TrackInteractions(*Evolution, MidPhaseAccessor);
+	TrackInteractions(PBDSolver, *Evolution, MidPhaseAccessor);
 
 	// Process the list of interactions that we built from the midphases
 	ProcessInteractions(*Evolution);
@@ -669,57 +685,61 @@ void FBuoyancySubsystemSimCallback::OnMidPhaseModification_Internal(Chaos::FMidP
 }
 
 void FBuoyancySubsystemSimCallback::TrackInteractions(
+	Chaos::FPBDRigidsSolver& PBDSolver,
 	Chaos::FPBDRigidsEvolution& Evolution,
 	Chaos::FMidPhaseModifierAccessor& MidPhaseAccessor)
 {
 	SCOPE_CYCLE_COUNTER(STAT_BuoyancySubsystem_VisitMidphases)
 
-	//
-	// This right here is the ugliest bit of the subsystem,
-	// since it loops needlessly over all midphases. It is
-	// tempting to think that this is the one most deserving
-	// of optimization, but the extra iterations here are far
-	// from the slowest part!
-	//
-
-	// Loop over all midphases
-	MidPhaseAccessor.VisitMidPhases([this, &Evolution](Chaos::FMidPhaseModifier& MidPhase)
+	// The spline userdatapt manager contains a PT array of all water splines indexed on
+	// water particle unique index. Therefore, we can use it to loop over all water body
+	// particles and check their midphases.
+	SplineData->VisitAllData_PT([this, &PBDSolver, &Evolution, &MidPhaseAccessor](Chaos::FUniqueIdx ParticleIdx, const TSharedPtr<FBuoyancyWaterSplineData>& WaterData)
 	{
-		// Make sure we have two valid particles
-		Chaos::FGeometryParticleHandle* WaterParticle;
-		Chaos::FGeometryParticleHandle* OtherParticle;
-		MidPhase.GetParticles(&WaterParticle, &OtherParticle);
-		if (WaterParticle == nullptr ||
-			OtherParticle == nullptr)
+		// If the WaterData shared ptr is invalid, skip this one
+		if (WaterData.IsValid() == false)
 		{
 			return;
 		}
 
-		// Get spline data for particle 0. If it exists, then it's water.
-		// If it doesn't exist, then try the other particle.
-		const TSharedPtr<FBuoyancyWaterSplineData>* WaterSpline = SplineData->GetData_PT(*WaterParticle);
-		if (WaterSpline == nullptr || !WaterSpline->IsValid())
+		// Get water body particle proxy from index
+		Chaos::FSingleParticlePhysicsProxy* ParticleProxy = PBDSolver.GetParticleProxy_PT(ParticleIdx);
+		if (ParticleProxy == nullptr || ParticleProxy->GetMarkedDeleted())
 		{
-			// Swap the particles and try again
-			Swap(WaterParticle, OtherParticle);
-			WaterSpline = SplineData->GetData_PT(*WaterParticle);
-			if (WaterSpline == nullptr || !WaterSpline->IsValid())
-			{
-				// Neither particle has a water spline data, so give up.
-				// This is not a water interaction
-				return;
-			}
+			return;
 		}
 
-		// Make sure the non-water particle is backed by a rigid
-		Chaos::FPBDRigidParticleHandle* RigidParticle = OtherParticle->CastToRigidParticle();
-		if (RigidParticle == nullptr)
+		// Get water body particle handle from proxy
+		Chaos::FGeometryParticleHandle* WaterParticle = ParticleProxy->GetHandle_LowLevel();
+		if (WaterParticle == nullptr)
 		{
 			return;
 		}
 
-		// Finally... we know for sure this is a midphase that we wanna process
-		TrackInteraction(Evolution, WaterParticle, RigidParticle, *WaterSpline->Get(), MidPhase);
+		// Get all midphases which involve this particle
+		Chaos::FMidPhaseModifierParticleRange MidPhases = MidPhaseAccessor.GetMidPhases(WaterParticle);
+
+		// Loop over all midphases
+		for (auto& MidPhase : MidPhases)
+		{
+			// Make sure we have two valid particles
+			Chaos::FGeometryParticleHandle* OtherParticle = MidPhase.GetOtherParticle(WaterParticle);
+			if (WaterParticle == nullptr ||
+				OtherParticle == nullptr)
+			{
+				continue;
+			}
+
+			// Make sure the non-water particle is backed by a rigid
+			Chaos::FPBDRigidParticleHandle* RigidParticle = OtherParticle->CastToRigidParticle();
+			if (RigidParticle == nullptr)
+			{
+				continue;
+			}
+
+			// Finally... we know for sure this is a midphase that we wanna process
+			TrackInteraction(Evolution, WaterParticle, RigidParticle, *WaterData.Get(), MidPhase);
+		}
 	});
 }
 
diff --git a/Engine/Plugins/Experimental/Buoyancy/Source/Runtime/Public/BuoyancySubsystem.h b/Engine/Plugins/Experimental/Buoyancy/Source/Runtime/Public/BuoyancySubsystem.h
--- a/Engine/Plugins/Experimental/Buoyancy/Source/Runtime/Public/BuoyancySubsystem.h
+++ b/Engine/Plugins/Experimental/Buoyancy/Source/Runtime/Public/BuoyancySubsystem.h
@@ -226,7 +226,7 @@ private:
 	virtual void OnPreSimulate_Internal() override;
 	virtual void OnMidPhaseModification_Internal(Chaos::FMidPhaseModifierAccessor& Modifier) override;
 
-	void TrackInteractions(Chaos::FPBDRigidsEvolution& Evolution, Chaos::FMidPhaseModifierAccessor& MidPhaseAccessor);
+	void TrackInteractions(Chaos::FPBDRigidsSolver& PBDSolver, Chaos::FPBDRigidsEvolution& Evolution, Chaos::FMidPhaseModifierAccessor& MidPhaseAccessor);
 	void TrackInteraction(Chaos::FPBDRigidsEvolution& Evolution, Chaos::FGeometryParticleHandle* WaterParticle, Chaos::FPBDRigidParticleHandle* RigidParticle, const FBuoyancyWaterSplineData& WaterSpline, Chaos::FMidPhaseModifier& MidPhase);
 	void ProcessInteractions(Chaos::FPBDRigidsEvolution& Evolution);
 	void ProcessInteraction(Chaos::FPBDRigidsEvolution& Evolution, FBuoyancyInteraction& Interaction);
diff --git a/Engine/Plugins/Experimental/ChaosUserDataPT/Source/ChaosUserDataPT/Public/ChaosUserDataPT.h b/Engine/Plugins/Experimental/ChaosUserDataPT/Source/ChaosUserDataPT/Public/ChaosUserDataPT.h
--- a/Engine/Plugins/Experimental/ChaosUserDataPT/Source/ChaosUserDataPT/Public/ChaosUserDataPT.h
+++ b/Engine/Plugins/Experimental/ChaosUserDataPT/Source/ChaosUserDataPT/Public/ChaosUserDataPT.h
@@ -57,6 +57,9 @@ namespace Chaos
 		// Flag for clearing all data
 		bool bClear = false;
 
+		// Flag for whether or not data resize is allowed from changes
+		bool bResizeOnClear = true;
+
 		// Monotonically increasing identifier for the input object. Each
 		// newly constructed input will store and increment this;
 		int32 Identifier = -1;
@@ -190,8 +193,21 @@ namespace Chaos
 			return UserDataMap_PT.IsValidIndex(Idx) ? UserDataMap_PT[Idx].Get() : nullptr;
 		}
 
+		// Execute a lambda on each data entry
+		void VisitAllData_PT(const TFunction<void(Chaos::FUniqueIdx, const TUserData&)>& Lambda)
+		{
+			for (auto It = UserDataMap_PT.CreateIterator(); It; ++It)
+			{
+				// If the unique ptr has data
+				if (It->IsValid())
+				{
+					Lambda(Chaos::FUniqueIdx(It.GetIndex()), *It->Get());
+				}
+			}
+		}
+
 		// Clear all data
-		bool ClearData_GT()
+		bool ClearData_GT(const bool bResize = true)
 		{
 			if (TInput* Input = this->GetProducerInputData_External())
 			{
@@ -202,6 +218,9 @@ namespace Chaos
 
 				// Mark the flag for clear all data
 				Input->bClear = true;
+
+				// Mark whether or not we're allowed to resize
+				Input->bResizeOnClear = bResize;
 			}
 
 			// Failed to queue removal
@@ -238,7 +257,12 @@ namespace Chaos
 
 						// Empty the userdata map
 						UserDataMap_PT.Empty();
-						UserDataMap_PT.Shrink();
+
+						// If resize is allowed, do it
+						if (Input->bResizeOnClear)
+						{
+							UserDataMap_PT.Shrink();
+						}
 					}
 
 					// Add new data
diff --git a/Engine/Source/Runtime/Experimental/Chaos/Private/Chaos/MidPhaseModification.cpp b/Engine/Source/Runtime/Experimental/Chaos/Private/Chaos/MidPhaseModification.cpp
--- a/Engine/Source/Runtime/Experimental/Chaos/Private/Chaos/MidPhaseModification.cpp
+++ b/Engine/Source/Runtime/Experimental/Chaos/Private/Chaos/MidPhaseModification.cpp
@@ -33,6 +33,24 @@ namespace Chaos
 		*Particle1 = MidPhase->GetParticle1();
 	}
 
+	FGeometryParticleHandle* FMidPhaseModifier::GetOtherParticle(const FGeometryParticleHandle* InParticle) 
+	{
+		if (MidPhase)
+		{
+			FGeometryParticleHandle* Particle0 = MidPhase->GetParticle0();
+			FGeometryParticleHandle* Particle1 = MidPhase->GetParticle1();
+			if (InParticle == Particle0)
+			{
+				return Particle1;
+			}
+			else if (InParticle == Particle1)
+			{
+				return Particle0;
+			}
+		}
+		return nullptr;
+	}
+
 	const FGeometryParticleHandle* FMidPhaseModifier::GetOtherParticle(const FGeometryParticleHandle* InParticle) const
 	{
 		if (MidPhase)
diff --git a/Engine/Source/Runtime/Experimental/Chaos/Public/Chaos/MidPhaseModification.h b/Engine/Source/Runtime/Experimental/Chaos/Public/Chaos/MidPhaseModification.h
--- a/Engine/Source/Runtime/Experimental/Chaos/Public/Chaos/MidPhaseModification.h
+++ b/Engine/Source/Runtime/Experimental/Chaos/Public/Chaos/MidPhaseModification.h
@@ -61,8 +61,9 @@ namespace Chaos
 
 		CHAOS_API void GetParticles(FGeometryParticleHandle** Particle0, FGeometryParticleHandle** Particle1) const;
 
-		CHAOS_API const FGeometryParticleHandle* GetOtherParticle(const FGeometryParticleHandle* InParticle) const;
+		CHAOS_API FGeometryParticleHandle* GetOtherParticle(const FGeometryParticleHandle* InParticle);
 
+		CHAOS_API const FGeometryParticleHandle* GetOtherParticle(const FGeometryParticleHandle* InParticle) const;
 
 	private:
 
