Allow for volume ratios in buoyancy algorithms where the particle's volume is greater than the shape's approximate volume. This fixes a case where buoyancy did not account for masses are increased independently from density.

More detail: previously when there were not enough wheels on a heavy build, the wheels would act EXTRA heavy in water instead of acting buoyant. This is due to a corrective scaling factor in the buoyancy algorithm which assumed that the volume which can be computed from the mass & density of a particle would always be exceeded by the sum of the volumes of the bounding boxes of the particle's shapes. This assumption doesn't hold in the case that we do mass redistribution for attached wheels. This change corrects the assumption.

diff --git a/Engine/Plugins/Experimental/Buoyancy/Source/Runtime/Private/BuoyancyAlgorithms.cpp b/Engine/Plugins/Experimental/Buoyancy/Source/Runtime/Private/BuoyancyAlgorithms.cpp
--- a/Engine/Plugins/Experimental/Buoyancy/Source/Runtime/Private/BuoyancyAlgorithms.cpp
+++ b/Engine/Plugins/Experimental/Buoyancy/Source/Runtime/Private/BuoyancyAlgorithms.cpp
@@ -20,6 +20,9 @@
 
 extern bool bBuoyancyDebugDraw;
 
+bool bBuoyancyAlgorithmsAllowVolRatioOverOne = false;
+FAutoConsoleVariableRef CVarBuoyancyAlgorithmsAllowVolRatioOverOne(TEXT("p.Buoyancy.Algorithms.AllowVolRatioOverOne"), bBuoyancyAlgorithmsAllowVolRatioOverOne, TEXT(""));
+
 
 //
 // Internal Functions
@@ -246,10 +249,15 @@ namespace BuoyancyAlgorithms
 		}
 
 		// Adjust the output volume based on the ratio of the material volume and the shape volume.
-		// We expect the shape volume to have overestimated the submerged volume for most shapes,
-		// especially those which are hollow.
+		//
+		// In most cases, we expect the shape volume to have overestimated the submerged volume for
+		// most shapes, especially those which are hollow.
+		//
+		// In some cases, if the mass of the submerged object has been changed independently of the
+		// density, which increases		
 		if (ParticleVol > UE_SMALL_NUMBER &&
-			ParticleVol < ShapeVol)
+			ShapeVol > UE_SMALL_NUMBER &&
+			(bBuoyancyAlgorithmsAllowVolRatioOverOne || ParticleVol < ShapeVol))
 		{
 			const float VolRatio = ParticleVol / ShapeVol;
 			SubmergedVol *= VolRatio;
