Add synctimestamp weakptr to buoyancy particle data so that deleted particles are not dereferenced when swapped from previous submersion data

diff --git a/Engine/Plugins/Experimental/Buoyancy/Source/Runtime/Private/BuoyancySubsystem.cpp b/Engine/Plugins/Experimental/Buoyancy/Source/Runtime/Private/BuoyancySubsystem.cpp
--- a/Engine/Plugins/Experimental/Buoyancy/Source/Runtime/Private/BuoyancySubsystem.cpp
+++ b/Engine/Plugins/Experimental/Buoyancy/Source/Runtime/Private/BuoyancySubsystem.cpp
@@ -946,6 +946,9 @@ void FBuoyancySubsystemSimCallback::ProcessInteraction(Chaos::FPBDRigidsEvolutio
 
 			Submersion.Particle = Interaction.RigidParticle;
 
+			Chaos::FSingleParticlePhysicsProxy* RigidProxy = static_cast<Chaos::FSingleParticlePhysicsProxy*>(Submersion.Particle->PhysicsProxy());
+			Submersion.SyncTimestamp = RigidProxy->GetSyncTimestamp();
+
 			// Get the weighted-average CoM
 			// NOTE: The unchecked division should be safe since we already
 			// know SubmergedVol > SMALL_NUMBER
@@ -984,7 +987,8 @@ void FBuoyancySubsystemSimCallback::ProcessInteraction(Chaos::FPBDRigidsEvolutio
 					// If we haven't already maxed out on water contacts, add one
 					if (MetaData.WaterContacts.Num() < FBuoyancySubmersionMetaData::MaxNumWaterContacts)
 					{
-						MetaData.WaterContacts.Add({ Interaction.WaterParticle, SubmergedVol, SubmergedCoM, CoMVel });
+						Chaos::FSingleParticlePhysicsProxy* WaterProxy = static_cast<Chaos::FSingleParticlePhysicsProxy*>(Interaction.WaterParticle->PhysicsProxy());
+						MetaData.WaterContacts.Add({ Interaction.WaterParticle, WaterProxy->GetSyncTimestamp(), SubmergedVol, SubmergedCoM, CoMVel});
 					}
 				}
 			}
@@ -1036,7 +1040,7 @@ void FBuoyancySubsystemSimCallback::ApplyBuoyantForces(Chaos::FPBDRigidsEvolutio
 
 namespace
 {
-	bool IsParticleValid(Chaos::FGeometryParticleHandle* ParticleHandle)
+	bool IsParticleValid(Chaos::FGeometryParticleHandle* ParticleHandle, TSharedPtr<FProxyTimestampBase, ESPMode::ThreadSafe> SyncTimestamp)
 	{
 		if (bBuoyancyCallbackDataParticleValidation == false)
 		{
@@ -1048,13 +1052,23 @@ namespace
 			return false;
 		}
 
+		if (SyncTimestamp.IsValid() == false)
+		{
+			return false;
+		}
+
+		if (SyncTimestamp->bDeleted)
+		{
+			return false;
+		}
+
 		IPhysicsProxyBase* Proxy = ParticleHandle->PhysicsProxy();
 		if (Proxy == nullptr)
 		{
 			return false;
 		}
 
-		if (Proxy->GetMarkedDeleted())
+		if (!ensureMsgf(Proxy->GetSyncTimestamp() == SyncTimestamp, TEXT("particle's sync timestamp doesn't match the sync timestamp that was passed to the particle validity check!")))
 		{
 			return false;
 		}
@@ -1117,7 +1131,7 @@ void FBuoyancySubsystemSimCallback::GenerateCallbackData()
 				continue;
 			}
 
-			if (!ensureMsgf(IsParticleValid(Submersion.Particle), TEXT("Submersion data for buoyancy callback includes invalid submerged particle handle")))
+			if (!ensureMsgf(IsParticleValid(Submersion.Particle, Submersion.SyncTimestamp.Pin()), TEXT("Submersion data for buoyancy callback includes invalid submerged particle handle")))
 			{
 				continue;
 			}
@@ -1125,7 +1139,7 @@ void FBuoyancySubsystemSimCallback::GenerateCallbackData()
 			// Build up output of new and continuing surface touches
 			for (const FBuoyancySubmersionMetaData::FWaterContact& WaterContact : MetaData.WaterContacts)
 			{
-				if (!ensureMsgf(IsParticleValid(WaterContact.Water), TEXT("Submersion data for buoyancy callback includes invalid water body particle handle")))
+				if (!ensureMsgf(IsParticleValid(WaterContact.Water, Submersion.SyncTimestamp.Pin()), TEXT("Submersion data for buoyancy callback includes invalid water body particle handle")))
 				{
 					continue;
 				}
@@ -1154,14 +1168,14 @@ void FBuoyancySubsystemSimCallback::GenerateCallbackData()
 				const int32 ObjectIndex = Iter.GetIndex();
 				const FBuoyancySubmersion& Submersion = PrevSubmersions[ObjectIndex];
 
-				if (!ensureMsgf(IsParticleValid(Submersion.Particle), TEXT("Previous frame submersion data for buoyancy callback includes invalid submerged particle handle")))
+				if (!IsParticleValid(Submersion.Particle, Submersion.SyncTimestamp.Pin()))
 				{
 					continue;
 				}
 
 				for (const FBuoyancySubmersionMetaData::FWaterContact& WaterContact : MetaData.WaterContacts)
 				{
-					if (!ensureMsgf(IsParticleValid(WaterContact.Water), TEXT("Previous frame submersion data for buoyancy callback includes invalid water body particle handle")))
+					if (!IsParticleValid(WaterContact.Water, WaterContact.SyncTimestamp.Pin()))
 					{
 						continue;
 					}
diff --git a/Engine/Plugins/Experimental/Buoyancy/Source/Runtime/Public/BuoyancyParticleData.h b/Engine/Plugins/Experimental/Buoyancy/Source/Runtime/Public/BuoyancyParticleData.h
--- a/Engine/Plugins/Experimental/Buoyancy/Source/Runtime/Public/BuoyancyParticleData.h
+++ b/Engine/Plugins/Experimental/Buoyancy/Source/Runtime/Public/BuoyancyParticleData.h
@@ -33,6 +33,7 @@ using FBuoyancyInteractionArray = TArray<FBuoyancyInteraction, TInlineAllocator<
 struct FBuoyancySubmersion
 {
 	Chaos::FPBDRigidParticleHandle* Particle = nullptr;
+	TWeakPtr<FProxyTimestampBase, ESPMode::ThreadSafe> SyncTimestamp = nullptr;
 	float Vol = 0.f;
 	FVector CoM = FVector::ZeroVector;
 	FVector Vel = FVector::ZeroVector;
@@ -45,6 +46,7 @@ struct FBuoyancySubmersionMetaData
 	struct FWaterContact
 	{
 		Chaos::FGeometryParticleHandle* Water = nullptr;
+		TWeakPtr<FProxyTimestampBase, ESPMode::ThreadSafe> SyncTimestamp = nullptr;
 		float Vol = 0.f;
 		FVector CoM = FVector::ZeroVector;
 		FVector Vel = FVector::ZeroVector;
