Buoyancy Plugin support for baked shallow water simulations
New algorithm that loops over faces to compute forces from water onto bodies
More accurate submerged volume computation

diff --git a/Engine/Plugins/Experimental/Buoyancy/Source/Runtime/Private/BuoyancyAlgorithms.cpp b/Engine/Plugins/Experimental/Buoyancy/Source/Runtime/Private/BuoyancyAlgorithms.cpp
--- a/Engine/Plugins/Experimental/Buoyancy/Source/Runtime/Private/BuoyancyAlgorithms.cpp
+++ b/Engine/Plugins/Experimental/Buoyancy/Source/Runtime/Private/BuoyancyAlgorithms.cpp
@@ -270,7 +270,7 @@ namespace BuoyancyAlgorithms
 		SCOPE_CYCLE_COUNTER(STAT_BuoyancyAlgorithms_ComputeSubmergedVolume)
 
 		if (ComputeSubmergedVolume(ParticleData, SubmergedParticle, WaterParticle, WaterX, WaterN, NumSubdivisions, MinVolume, SubmergedVol, SubmergedCoM))
-		{
+		{			
 			ScaleSubmergedVolume(Evolution, SubmergedParticle, SubmergedVol, TotalVol);
 
 
@@ -288,6 +288,7 @@ namespace BuoyancyAlgorithms
 		return false;
 	}
 
+
 	bool ComputeSubmergedVolume(FBuoyancyParticleData& ParticleData, const FGeometryParticleHandle* SubmergedParticle, const FGeometryParticleHandle* WaterParticle, const FVector& WaterX, const FVector& WaterN, int32 NumSubdivisions, float MinVolume, float& SubmergedVol, FVec3& SubmergedCoM)
 	{
 		// Get some initial data about the submerged particle
@@ -348,14 +349,14 @@ namespace BuoyancyAlgorithms
 			const FAABB3 WorldBox = LocalBox.TransformedAABB(ShapeWorldTransform);
 
 #if ENABLE_DRAW_DEBUG
-			if (bBuoyancyDebugDraw)
-			{
-				Chaos::FDebugDrawQueue::GetInstance().DrawDebugBox(
-					ShapeWorldTransform.TransformPosition(LocalBox.GetCenter()),
-					LocalBox.Extents() * .5f,
-					ShapeWorldTransform.GetRotation(),
-					FColor::Green, false, -1.f, SDPG_Foreground, 1.f);
-			}
+			//if (bBuoyancyDebugDraw)
+			//{
+			//	Chaos::FDebugDrawQueue::GetInstance().DrawDebugBox(
+			//		ShapeWorldTransform.TransformPosition(LocalBox.GetCenter()),
+			//		LocalBox.Extents() * .5f,
+			//		ShapeWorldTransform.GetRotation(),
+			//		FColor::Green, false, -1.f, SDPG_Foreground, 1.f);
+			//}
 #endif
 
 			// Get the world space position of the shape
@@ -401,11 +402,11 @@ namespace BuoyancyAlgorithms
 #if ENABLE_DRAW_DEBUG
 					if (bBuoyancyDebugDraw)
 					{
-						Chaos::FDebugDrawQueue::GetInstance().DrawDebugBox(
-							ShapeWorldTransform.TransformPosition(SubmergedBox.GetCenter()),
-							SubmergedBox.Extents() * .5f,
-							ShapeWorldTransform.GetRotation(),
-							FColor::Red, false, -1.f, SDPG_Foreground, 1.f);
+						//Chaos::FDebugDrawQueue::GetInstance().DrawDebugBox(
+//							ShapeWorldTransform.TransformPosition(SubmergedBox.GetCenter()),
+							//SubmergedBox.Extents() * .5f,
+							//ShapeWorldTransform.GetRotation(),
+							//FColor::Red, false, -1.f, SDPG_Foreground, 1.f);
 					}
 #endif
 				}
@@ -602,7 +603,7 @@ namespace BuoyancyAlgorithms
 		// Compute water drag force
 		//
 		// NOTE: This is a very approximate "ether drag" style model here, probably 
-		// we should scale the model with submerged volume for more accuracy,
+		// we should scale the model with submerged volume for more accuracy,WorldBox
 		// and apply the drag force in opposition to the linear motion of the submerged
 		// center of mass.
 		const float DragFactor = FMath::Max(0.f, 1.f - (WaterDrag * DeltaSeconds));
@@ -614,4 +615,649 @@ namespace BuoyancyAlgorithms
 		//
 		return true;
 	}
+		
+	void ComputeSubmergedVolumeAndForcesForParticle(FBuoyancyParticleData& ParticleData, 
+		const Chaos::FGeometryParticleHandle* SubmergedParticle, const FGeometryParticleHandle* WaterParticle, 
+		const FShallowWaterSimulationGrid& ShallowWaterGrid,
+		const Chaos::FPBDRigidsEvolution& Evolution, const float DeltaSeconds, const float WaterDensity, const float WaterDrag,
+		float &OutTotalSubmergedVol, Chaos::FVec3 &OutTotalSubmergedCoM, Chaos::FVec3 &OutTotalForce, Chaos::FVec3 &OutTotalTorque)
+	{		
+		SCOPE_CYCLE_COUNTER(STAT_BuoyancyAlgorithms_ComputeSubmergedVolumeAndForcesForParticle)
+
+		// Initialize cumulative values
+		OutTotalSubmergedVol = 0;
+		OutTotalSubmergedCoM = FVec3::ZeroVector;
+		OutTotalForce = FVec3::ZeroVector;
+		OutTotalTorque = FVec3::ZeroVector;
+			
+		// Get some initial data about the submerged particle
+		const FImplicitObject* RootImplicit = SubmergedParticle->GetGeometry();
+		const FShapeInstanceArray& ShapeInstances = SubmergedParticle->ShapeInstances();
+		const FConstGenericParticleHandle SubmergedGeneric = SubmergedParticle;
+		const FRigidTransform3 ParticleWorldTransform = SubmergedGeneric->GetTransformPQ();
+		const int32 ParticleIndex = ParticleData.GetIndex(*SubmergedParticle);
+		TSparseArray<TBitArray<>>& SubmergedShapes = ParticleData.SubmergedShapes;
+
+		// Some info about the water
+		const FImplicitObject* WaterRootImplicit = WaterParticle->GetGeometry();
+		const EImplicitObjectType WaterShapeType = Chaos::Private::GetImplicitCollisionType(WaterParticle, WaterRootImplicit);
+		const FShapeInstanceArray& WaterShapeInstances = WaterParticle->ShapeInstances();
+		const FShapeInstance* WaterShapeInstance = WaterShapeInstances[0].Get();
+
+		// Traverse the submerged particle's leaves
+		RootImplicit->VisitLeafObjects(
+			[&Evolution, DeltaSeconds, WaterDensity, WaterDrag, SubmergedParticle, &ShallowWaterGrid, 
+			ParticleIndex, &ShapeInstances, WaterShapeType, WaterShapeInstance, &ParticleWorldTransform, 
+			&SubmergedShapes, &OutTotalSubmergedVol, &OutTotalSubmergedCoM,
+			&OutTotalForce, &OutTotalTorque]
+			(const FImplicitObject* Implicit, const FRigidTransform3& RelativeTransform, const int32 RootObjectIndex, const int32 ObjectIndex, const int32 LeafObjectIndex)
+			{				
+				const FAABB3 RelativeBounds = Implicit->CalculateTransformedBounds(RelativeTransform);
+				const int32 ShapeIndex = (ShapeInstances.IsValidIndex(RootObjectIndex)) ? RootObjectIndex : 0;
+				const FShapeInstance* ShapeInstance = ShapeInstances[ShapeIndex].Get();
+				const EImplicitObjectType ShapeType = Chaos::Private::GetImplicitCollisionType(SubmergedParticle, Implicit);
+
+				// If this shape pair doesn't pass a narrow phase test then skip it
+				if (!ShapePairNarrowPhaseFilter(ShapeType, ShapeInstance, WaterShapeType, WaterShapeInstance))
+				{
+					return;
+				}
+
+				// If this shape has already been submerged, skip it to avoid double-counting
+				// any buoyancy contributions.
+				if (IsShapeSubmerged_Internal(SubmergedShapes, ParticleIndex, ObjectIndex))
+				{
+					return;
+				}
+				
+				FVec3 OutWaterP, OutWaterN;
+				float OutSubmergedVol;
+				FVec3 OutSubmergedCoM;
+				FVec3 OutForce, OutTorque;
+				ComputeSubmergedVolumeAndForcesForShape(SubmergedParticle,
+					Evolution, DeltaSeconds, WaterDensity, WaterDrag,
+					Implicit, RelativeTransform, RootObjectIndex, ObjectIndex, LeafObjectIndex,
+					ShallowWaterGrid,
+					OutWaterP, OutWaterN,
+					OutSubmergedVol, OutSubmergedCoM,
+					OutForce, OutTorque);
+
+				// mark the shape as submerged
+				if (OutSubmergedVol > SMALL_NUMBER)
+				{
+					OutTotalForce += OutForce;
+					OutTotalTorque += OutTorque;
+					OutTotalSubmergedCoM += OutSubmergedCoM * OutSubmergedVol;
+
+					OutTotalSubmergedVol += OutSubmergedVol;
+
+					SubmergeShape_Internal(SubmergedShapes, ParticleIndex, ObjectIndex);
+				}
+			});
+
+		// compute final force and torque for particle as weighted average of all the shapes
+		OutTotalSubmergedCoM /= OutTotalSubmergedVol;
+	}
+
+	void ComputeSubmergedVolumeAndForcesForShape(const Chaos::FGeometryParticleHandle* SubmergedParticle,
+		const Chaos::FPBDRigidsEvolution& Evolution, float DeltaSeconds, const float WaterDensity, const float WaterDrag,
+		const Chaos::FImplicitObject* Implicit, const Chaos::FRigidTransform3& RelativeTransform, const int32 RootObjectIndex, const int32 ObjectIndex, const int32 LeafObjectIndex,
+		const FShallowWaterSimulationGrid& ShallowWaterGrid,		
+		Chaos::FVec3& OutWaterP, Chaos::FVec3& OutWaterN,
+		float& OutSubmergedVol, Chaos::FVec3& OutSubmergedCoM,
+		Chaos::FVec3& OutForce, Chaos::FVec3& OutTorque)
+	{
+		SCOPE_CYCLE_COUNTER(STAT_BuoyancyAlgorithms_ComputeSubmergedVolumeAndForcesForShape)
+
+		const FConstGenericParticleHandle SubmergedGeneric = SubmergedParticle;
+		const FRigidTransform3 ParticleWorldTransform = SubmergedGeneric->GetTransformPQ();
+
+		const FPBDRigidParticleHandle* RigidParticle = SubmergedParticle->CastToRigidParticle();
+
+		// Get the world-space bounds of shape A
+		FRigidTransform3 ShapeWorldTransform = RelativeTransform * ParticleWorldTransform;
+		FAABB3 LocalBox = Implicit->BoundingBox();
+		if (const FImplicitSphere3* Sphere = Implicit->AsA<FImplicitSphere3>())
+		{
+			// If we have a sphere, ignore rotation because submerged volume is independent of rotation
+			// and also we don't want to apply any torques on the wheel.			
+			const FVec3 SphereCenter = ShapeWorldTransform.TransformPosition(Sphere->GetCenter());
+			const FVec3 SphereExtent = FVec3(Sphere->GetRadius());
+			LocalBox = FAABB3(-SphereExtent, SphereExtent);
+			ShapeWorldTransform.SetTranslation(SphereCenter);
+
+			// #todo(dmp): at this point, we haven't figured out the water normal so we use up vector for now
+			FVec3 TmpWaterN = {0, 0, 1};
+			ShapeWorldTransform.SetRotation(FRotationMatrix::MakeFromZ(TmpWaterN).ToQuat());
+		}
+				
+		// water plane parameters for the body.  Note we have 1 plane per box
+		OutWaterP = FVec3(0, 0, 0);
+		OutWaterN = FVec3(0, 0, 1);
+
+		// Compute submerged volume and CoM
+		OutSubmergedVol = 0;
+		OutSubmergedCoM = FVec3(0, 0, 0);
+
+		// output forces
+		OutForce = FVec3(0, 0, 0);
+		OutTorque = FVec3(0, 0, 0);
+
+		float SubmergedCOMTotalWeight = 0;
+		
+		FVector WorldBoxCenter = ShapeWorldTransform.TransformPosition(LocalBox.GetCenter());
+
+		// #todo(dmp): just store normals in the shallow water texture and avoid all this
+		// water plane position is computed from the centroid of the box		
+		FVector CenterVelocity;
+		float CenterHeight;
+		float CenterDepth;
+		ShallowWaterGrid.SampleShallowWaterSimulationAtPosition(WorldBoxCenter, CenterVelocity, CenterHeight, CenterDepth);
+
+		// no water at shape centroid of shape so we return
+		if (CenterDepth < SMALL_NUMBER)
+		{
+			return;
+		}
+
+		// #todo(dmp): lift normal computation out of this to make function more generic
+		// water plane is based on the shape's centroid
+		OutWaterP = WorldBoxCenter;
+		OutWaterP.Z = CenterHeight;
+		OutWaterN = ShallowWaterGrid.ComputeShallowWaterSimulationNormalAtPosition(OutWaterP);
+
+		// #todo(dmp): generalize to convex shapes and remove box assumption
+		const int TotalNumVertices = 8;
+
+		FVector WorldVertexPosition[TotalNumVertices];
+		bool VertexIsUnderwater[TotalNumVertices];
+
+		// keep track of a reference point inside the clipped geometry that is the average of all the 
+		// submerged vertices and intersection points this will jump around because it isn't weighted by the area/volume
+		// but is sufficient for tet volume computations since we only need an interior point
+		FVector InteriorRefPoint(0, 0, 0);
+		int InteriorRefPointCount = 0;
+
+		// cache water values for each vertex and determine broadphase of shape interacting with the water	
+		for (int i = 0; i < TotalNumVertices; ++i)
+		{
+			// compute world space position of vertex
+			const FVector CurrWorldVertexPosition = ShapeWorldTransform.TransformPosition(LocalBox.GetVertex(i));
+			WorldVertexPosition[i] = CurrWorldVertexPosition;
+
+			if ((CurrWorldVertexPosition - OutWaterP).Dot(OutWaterN) < 0)
+			{
+				VertexIsUnderwater[i] = true;
+				InteriorRefPoint += WorldVertexPosition[i];
+				InteriorRefPointCount++;
+			}
+			else
+			{
+				VertexIsUnderwater[i] = false;
+			}
+		}
+
+		// find intersection points of the box and water plane
+		FVector IntersectionCenter;			
+		TArray<FVector> OrderedIntersectionPoints;
+		TMap<int, FVector> AllIntersectionPoints;		
+		BuoyancyAlgorithms::FindAllIntersectionPoints(OutWaterP, OutWaterN, LocalBox, WorldVertexPosition, AllIntersectionPoints, OrderedIntersectionPoints, IntersectionCenter);
+		const int NumIntersectionPoints = AllIntersectionPoints.Num();
+		
+		// Add intersection points to interior reference point
+		// values are already accumulated into the intersection center
+		InteriorRefPoint += IntersectionCenter * NumIntersectionPoints;
+		InteriorRefPointCount += NumIntersectionPoints;
+		InteriorRefPoint /= InteriorRefPointCount;
+
+		float AccumulatedForceWeight = 0;
+		
+		const FVec3 WorldCoM = SubmergedGeneric->PCom();
+
+		// Sum up submerged areas and volumes for each face
+		const int NumFaces = 6;
+		for (int FaceIdx = 0; FaceIdx < NumFaces; ++FaceIdx)
+		{
+			const Chaos::FAABBFace CurrFace = LocalBox.GetFace(FaceIdx);
+
+			// compute area of unsubmerged face	
+			// #todo(dmp): precompute edge lengths
+			const FVector Pos0 = WorldVertexPosition[CurrFace.VertexIndex[0]];
+			const FVector Pos1 = WorldVertexPosition[CurrFace.VertexIndex[1]];
+			const FVector Pos2 = WorldVertexPosition[CurrFace.VertexIndex[2]];
+			const float FaceArea = (Pos1 - Pos0).Length() * (Pos2 - Pos1).Length();
+
+			// skip small faces from degenerate bodies
+			if (FaceArea < SMALL_NUMBER)
+			{
+				continue;
+			}
+
+			FVector IntersectedFaceCenter(0, 0, 0);
+
+			// we know that a given face can have at most 5 points below the surface
+			// #todo(dmp): we need to remove this assumption if we support convex hulls or other types of geometry
+			static constexpr int32 MaxNumVerticesBelowWater = 5;
+			TArray<FVector, TInlineAllocator<MaxNumVerticesBelowWater>> FaceIntersectionPoints;
+
+			// walk vertices for the current face in counter clockwise order and find intersections and
+			// submerged vertices
+			const int NumVerticesForFace = 4;
+			for (int i = 0; i < NumVerticesForFace; ++i)
+			{
+				// get the current edge index belonging to the i and i+1 vertices
+				// note this is a bit weird...we could just ask the edge for the vertices but they
+				// might be in the wrong order for correctly constructing the intersected face
+				const int CurrEdgeIdx = CurrFace.EdgeIndex[i];
+
+				// get the indices of the two vertices for the current edge
+				const int Edge0Idx = CurrFace.VertexIndex[i];
+				const int Edge1Idx = CurrFace.VertexIndex[(i + 1) % 4];
+				const FVector V0 = WorldVertexPosition[Edge0Idx];
+				const FVector V1 = WorldVertexPosition[Edge1Idx];
+
+				// if there is an intersection for this edge, add it to the list for the current face
+				const FVector* EdgeIntersection = AllIntersectionPoints.Find(CurrEdgeIdx);
+				if (EdgeIntersection)
+				{
+					FaceIntersectionPoints.Add(*EdgeIntersection);
+					IntersectedFaceCenter += *EdgeIntersection;
+				}
+
+				// add the second endpoint if it is underwater
+				if (VertexIsUnderwater[Edge1Idx])
+				{
+					FaceIntersectionPoints.Add(V1);
+					IntersectedFaceCenter += V1;
+				}
+			}
+			
+			const int NumFaceIntersectionPts = FaceIntersectionPoints.Num();
+
+			// if the face intersects the water or is fully submerged, compute face coverage
+			if (NumFaceIntersectionPts > 0)
+			{
+				IntersectedFaceCenter /= NumFaceIntersectionPts;
+
+				// we have a list of the intersection points for a face in correct counterclockwise order
+				// fan layout for points on face, sum up areas and volume for the current face
+				for (int i = 0; i < NumFaceIntersectionPts; i++)
+				{
+					const FVector V0 = FaceIntersectionPoints[i];
+					const FVector V1 = FaceIntersectionPoints[(i + 1) % NumFaceIntersectionPts];
+					const FVector V2 = IntersectedFaceCenter;
+
+					FVector TriBaryCenter;
+					float TriArea;
+					float TetVolume;
+					FVector TriNormal;
+					BuoyancyAlgorithms::ComputeTriangleAreaAndVolume(V0, V1, V2, InteriorRefPoint, TriBaryCenter, TriNormal, TriArea, TetVolume, bBuoyancyDebugDraw);
+
+					// skip degenerate triangles with 0 area
+					if (TriArea < SMALL_NUMBER)
+					{
+						continue;
+					}
+
+					OutSubmergedCoM += TriBaryCenter * TriArea;
+					SubmergedCOMTotalWeight += TriArea;
+					OutSubmergedVol += TetVolume;
+
+					/////
+					// sum up forces/torques on the triangle
+					/////								
+
+					// sample velocity at center of triangle
+					// #todo(dmp): this should just be interpolated from vertex velocities to avoid so many samples	
+					// #todo(dmp): split arrays out to store velocity separate from other values
+					FVector CurrWaterVelocity;
+					float TmpWaterVal;
+					ShallowWaterGrid.SampleShallowWaterSimulationAtPosition(TriBaryCenter,
+						CurrWaterVelocity, TmpWaterVal, TmpWaterVal);
+
+					FVector TotalWorldForce;
+					FVector TotalWorldTorque;
+
+					// #todo(dmp): we scale drag to make it closer to the old algorithm for now
+					const float DragFactor = .1 * WaterDrag;
+					BuoyancyAlgorithms::ComputeFluidForceForTriangle(
+						DragFactor, DeltaSeconds, RigidParticle, WorldCoM, TriBaryCenter, TriNormal, TriArea, TetVolume,
+						CurrWaterVelocity, OutWaterP, OutWaterN, TotalWorldForce, TotalWorldTorque);
+
+					OutForce += TotalWorldForce * TriArea;
+					OutTorque += TotalWorldTorque * TriArea;
+					AccumulatedForceWeight += TriArea;
+
+#if ENABLE_DRAW_DEBUG
+					if (bBuoyancyDebugDraw)
+					{
+						FVector CurrWaterVelocityViz = CurrWaterVelocity;
+						CurrWaterVelocityViz.Normalize();
+						Chaos::FDebugDrawQueue::GetInstance().DrawDebugDirectionalArrow(TriBaryCenter, TriBaryCenter + CurrWaterVelocityViz * 20, 20, FColor::Red, false, -1.f, -1, 2.f);
+						Chaos::FDebugDrawQueue::GetInstance().DrawDebugDirectionalArrow(TriBaryCenter, TriBaryCenter + TotalWorldForce * .01, 20, FColor::Green, false, -1.f, -1, 2.f);
+					}
+#endif					
+				}
+			}
+		}
+
+		// Sum up tet volumes from the fan of triangles across the intersection face
+		int NumOrderedIntersectionPoints = OrderedIntersectionPoints.Num();
+		for (int i = 0; i < NumOrderedIntersectionPoints; i++)
+		{
+			FVector V0 = OrderedIntersectionPoints[i];
+			FVector V1 = OrderedIntersectionPoints[(i + 1) % NumOrderedIntersectionPoints];
+			FVector V2 = IntersectionCenter;
+
+			FVector TriBaryCenter;
+			float TriArea;
+			float TetVolume;
+			FVector TriNormal;
+			BuoyancyAlgorithms::ComputeTriangleAreaAndVolume(V0, V1, V2, InteriorRefPoint, TriBaryCenter, TriNormal, TriArea, TetVolume, bBuoyancyDebugDraw);
+
+			// skip degenerate triangles with 0 area
+			if (TriArea < SMALL_NUMBER)
+			{
+				continue;
+			}
+
+			OutSubmergedCoM += TriBaryCenter * TriArea;
+			SubmergedCOMTotalWeight += TriArea;
+			OutSubmergedVol += TetVolume;
+		}
+
+		// normalize weighted average of contributions to world submerged COM
+		// return if nothing has contributed to CoM
+		if (SubmergedCOMTotalWeight > SMALL_NUMBER)
+		{
+			OutSubmergedCoM /= SubmergedCOMTotalWeight;
+		}
+		else
+		{
+			return;
+		}
+
+		// #todo(dmp): is it correct to normalize by the accumulated area over which the force is done?  Seems wrong...
+		if (AccumulatedForceWeight > SMALL_NUMBER)
+		{
+			OutForce /= AccumulatedForceWeight;
+			OutTorque /= AccumulatedForceWeight;
+		}
+
+		// compute buoyancy force and delta velocity for solver
+		Chaos::FVec3 WorldBuoyantForce, WorldBuoyantTorque;
+		BuoyancyAlgorithms::ComputeBuoyantForceForShape(Evolution, RigidParticle, DeltaSeconds,
+			WaterDensity, OutSubmergedCoM, OutSubmergedVol, OutWaterN, WorldBuoyantForce, WorldBuoyantTorque);
+
+		// Total force acting on the particle is the buoyant force plus all the forces due to matching the fluid velocity
+		OutForce += WorldBuoyantForce;
+		OutTorque += WorldBuoyantTorque;
+
+#if ENABLE_DRAW_DEBUG
+		if (bBuoyancyDebugDraw)
+		{
+			Chaos::FDebugDrawQueue::GetInstance().DrawDebugSphere(WorldBoxCenter, 10, 10, FColor::Green, false, -1.f, -1, 2.0f);
+			Chaos::FDebugDrawQueue::GetInstance().DrawDebugSphere(InteriorRefPoint, 10, 10, FColor::Yellow, false, -1.f, -1, 2.0f);
+			Chaos::FDebugDrawQueue::GetInstance().DrawDebugSphere(OutWaterP, 10, 10, FColor::Orange, false, -1.f, -1, 2.0f);
+			Chaos::FDebugDrawQueue::GetInstance().DrawDebugDirectionalArrow(OutWaterP, OutWaterP + OutWaterN * 100, 20, FColor::Orange, false, -1.f, -1, 6.f);
+			Chaos::FDebugDrawQueue::GetInstance().DrawDebugSphere(OutSubmergedCoM, 10, 10, FColor::Magenta, false, -1.f, -1, 2.0f);
+
+			//float SubmergedVol1;
+			//Chaos::FVec3 SubmergedCoM1;
+			//float TotalVol1;
+			//if (BuoyancyAlgorithms::ComputeSubmergedVolume(BuoyancyParticleData, Evolution, Interaction.RigidParticle, Interaction.WaterParticle, Interaction.ClosestPoint, WaterN, BuoyancySettings->MaxNumBoundsSubdivisions, BuoyancySettings->MinBoundsSubdivisionVol, SubmergedVol1, SubmergedCoM1, TotalVol1))
+			//{
+			//}
+			//Chaos::FDebugDrawQueue::GetInstance().DrawDebugSphere(SubmergedCoM1, 10, 10, FColor::Black, false, -1.f, -1, 2.0f);
+
+			Chaos::FDebugDrawQueue::GetInstance().DrawDebugBox(
+				ShapeWorldTransform.TransformPosition(LocalBox.GetCenter()),
+				LocalBox.Extents() * .5f,
+				ShapeWorldTransform.GetRotation(),
+				FColor::Red, false, -1.f, SDPG_Foreground, 3.f);
+		}
+#endif	
+
+	}
+
+	void FindAllIntersectionPoints(const Chaos::FVec3& WaterP, const Chaos::FVec3& WaterN, const Chaos::FAABB3& LocalBox, const FVector WorldVertexPosition[8],
+		TMap<int, FVector>& OutEdgeToIntersectionPointMap, TArray<FVector>& OutOrderedIntersectionPoints, FVector &OutIntersectionCenter)
+	{		
+		SCOPE_CYCLE_COUNTER(STAT_BuoyancyAlgorithms_FindAllIntersectionPoints)
+
+		OutIntersectionCenter = FVector(0, 0, 0);
+
+		// loop over all edges and compute intersection points
+		for (int CurrEdgeIdx = 0; CurrEdgeIdx < 12; ++CurrEdgeIdx)
+		{
+			const FAABBEdge CurrEdge = LocalBox.GetEdge(CurrEdgeIdx);
+
+			// edge plane intersection for each edge in counterclockwise order
+			const FVector RayOrigin = WorldVertexPosition[CurrEdge.VertexIndex0];
+			const FVector RayEndPoint = WorldVertexPosition[CurrEdge.VertexIndex1];			
+			const FVector RayDir = RayEndPoint - RayOrigin;
+
+			// if the water plane intersects the current edge, add intersection point it to the list
+			const float DirDotN = RayDir.Dot(WaterN);
+			if (FMath::Abs(DirDotN) > SMALL_NUMBER)
+			{
+				const float t = (WaterP - RayOrigin).Dot(WaterN) / DirDotN;
+				if (t >= 0 && t <= 1)
+				{
+					const FVector IntersectionPoint = RayOrigin + RayDir * t;
+
+					OutEdgeToIntersectionPointMap.Add(CurrEdgeIdx, IntersectionPoint);
+
+					OutIntersectionCenter += IntersectionPoint;
+				}
+			}
+		}
+
+		if (OutEdgeToIntersectionPointMap.Num() > 0)
+		{
+			OutIntersectionCenter /= OutEdgeToIntersectionPointMap.Num();
+
+			// Sort intersection points by angle around reference point
+			BuoyancyAlgorithms::SortIntersectionPointsByAngle(WaterP, WaterN, OutIntersectionCenter, OutEdgeToIntersectionPointMap, OutOrderedIntersectionPoints);
+		}
+	}
+
+	void SortIntersectionPointsByAngle(const Chaos::FVec3& WaterP, const Chaos::FVec3& WaterN, const Chaos::FVec3& IntersectionCenter, const TMap<int, FVector>& EdgeToIntersectionPointMap, TArray<FVector>& OutOrderedIntersectionPoints)
+	{
+		SCOPE_CYCLE_COUNTER(STAT_BuoyancyAlgorithms_SortIntersectionPointsByAngle)
+
+		// Order intersection points around intersection center
+		TMap<float, FVector> IntersectionPointAngleToCenter;
+
+		auto IntersectionPointMapIt = EdgeToIntersectionPointMap.CreateConstIterator();
+		
+		// first element is assigned to be angle 0
+		FVector CurrPoint = IntersectionPointMapIt.Value();
+		FVector CurrVector = CurrPoint - IntersectionCenter;
+		CurrVector.Normalize();
+
+		const FVector AngleStart = CurrVector;
+		IntersectionPointAngleToCenter.Add(0, CurrPoint);
+
+		// loop over remaining points and compute angles to the first point
+		++IntersectionPointMapIt;
+		for (; IntersectionPointMapIt; ++IntersectionPointMapIt)
+		{
+			CurrPoint = IntersectionPointMapIt.Value();
+			
+			CurrVector = CurrPoint - IntersectionCenter;
+			CurrVector.Normalize();
+			
+			float CurrAngle = FMath::Acos(AngleStart.Dot(CurrVector));
+
+			if (AngleStart.Cross(CurrVector).Dot(WaterN) < 0)
+			{
+				CurrAngle = 2.f * PI - CurrAngle;
+			}
+
+			IntersectionPointAngleToCenter.Add(CurrAngle, CurrPoint);
+		}
+
+		// #todo(dmp): replace this sort with something faster since we only have between 3-6 points to sort (and index 0 is already sorted)
+		// order intersection points by the angle it makes so we have a correct winding order
+		// for tesselation
+		IntersectionPointAngleToCenter.KeySort([](const float& A, const float& B) {
+			return A < B;
+			});
+
+		// Generate an array of all intersection points		
+		IntersectionPointAngleToCenter.GenerateValueArray(OutOrderedIntersectionPoints);
+	}
+
+	void ComputeTriangleAreaAndVolume(const FVector &V0, const FVector &V1, const FVector &V2, const FVector &MeshCenterPoint, FVector &OutTriangleBaryCenter, FVector &OutNormal, float& OutArea, float& OutVolume, bool DebugDraw /*= false*/)
+	{
+		SCOPE_CYCLE_COUNTER(STAT_BuoyancyAlgorithms_ComputeTriangleAreaAndVolume)
+
+		// compute center of triangle
+		OutTriangleBaryCenter = (V0 + V1 + V2) / 3.;
+
+		// add up the area of this triangle for all the faces		 
+		OutArea = .5 * (V2 - V0).Cross(V1 - V0).Length();
+
+		// add up the volume of the tet created by this triangle and the center of the box
+		const FVector A = V0 - MeshCenterPoint;
+		const FVector B = V1 - MeshCenterPoint;
+		const FVector C = V2 - MeshCenterPoint;
+		const FVector N = A.Cross(B);
+
+		OutVolume = FMath::Abs((N.Dot(C)) / 6.f);
+
+		OutNormal = (V1 - V0).Cross(V2 - V0);
+		OutNormal.Normalize();
+
+#if ENABLE_DRAW_DEBUG
+		if (DebugDraw)
+		{
+			Chaos::FDebugDrawQueue::GetInstance().DrawDebugLine(V0, V1, FColor::Orange, false, -1.f, 0, 1.f);
+			Chaos::FDebugDrawQueue::GetInstance().DrawDebugLine(V1, V2, FColor::Orange, false, -1.f, 0, 1.f);
+			Chaos::FDebugDrawQueue::GetInstance().DrawDebugLine(V2, V0, FColor::Orange, false, -1.f, 0, 1.f);
+
+			Chaos::FDebugDrawQueue::GetInstance().DrawDebugSphere(OutTriangleBaryCenter, 5, 10, FColor::Cyan, false, -1.f, -1, 1.f);
+			Chaos::FDebugDrawQueue::GetInstance().DrawDebugDirectionalArrow(OutTriangleBaryCenter, OutTriangleBaryCenter + OutNormal * 20, 20, FColor::Cyan, false, -1.f, -1, 2.f);
+		}
+#endif
+	}
+
+	// #todo(dmp): optimization- precalculate particle v, w, m given they never change between calls to compute forces
+	void ComputeFluidForceForTriangle(const float WaterDrag,
+		const float DeltaSeconds,
+		const Chaos::FPBDRigidParticleHandle* RigidParticle, const FVector WorldCoM,
+		const FVector &TriBaryCenter, const FVector &TriNormal, const float TriArea, const float TetVolume,
+		const FVector &WaterVelocity, const FVector &WaterP, const FVector &WaterN,
+		FVector &OutTotalWorldForce, FVector &OutTotalWorldTorque)
+	{
+		SCOPE_CYCLE_COUNTER(STAT_BuoyancyAlgorithms_ComputeFluidForceForTriangle)
+
+		OutTotalWorldForce = FVector(0, 0, 0);
+		OutTotalWorldTorque = FVector(0, 0, 0);
+
+		const FVec3 WorldForcePosition = TriBaryCenter;
+		const FVec3 WorldCOMToForcePos = WorldForcePosition - WorldCoM;
+
+		// project velocity sample onto water plane to support waterfalls and flowing rivers more accurately
+		const FVector WaterVelocityOnPlane = WaterVelocity - WaterVelocity.Dot(WaterN) * WaterN;
+
+		// Get world space particle linear velocity at current point
+		// note we are including the linear velocity from torque so objects spin properly in flow
+		const FVec3 SubmergedParticleVelocity = RigidParticle->GetV() + FVec3::CrossProduct(RigidParticle->GetW(), WorldCOMToForcePos);
+
+		// compute force and torque to set linear velocity to fluid velocity				
+		const FVec3 RelativeVelocity = WaterVelocityOnPlane - SubmergedParticleVelocity;
+		const float RelativeVelocityMag = RelativeVelocity.Length();
+
+		if (RelativeVelocityMag < SMALL_NUMBER)
+		{
+			return;
+		}
+		
+		const FVec3 AccelerationToSetVelocity = RelativeVelocity / DeltaSeconds;
+		FVec3 ForceFromWaterVelocity = AccelerationToSetVelocity * RigidParticle->M();
+
+		// #todo(dmp): make friction/drag/other properties more of a material property
+
+		// attenuate based on angle to inward facing tri normal					
+		// we reject angles > 90 since that would mean we are hitting a backface and no force should be applied					
+		// default to 1 if we have a still pool of water, we still want this to have a valid response so we don't have
+		// a discontinuity.
+		float FacingTest = 1.f;
+
+		FacingTest = RelativeVelocity.Dot(-TriNormal) / RelativeVelocityMag;
+		FacingTest = FacingTest > 0.f ? FacingTest : 0.f;
+
+		// apply some drag based on facing ratio...could make this more complex
+		ForceFromWaterVelocity *= FacingTest;
+
+		// #todo(dmp): drag should include a term based on relative velocity?
+		// ForceFromWaterVelocity *= RelativeVelocity.Length();
+
+		// drag coefficient			
+		ForceFromWaterVelocity *= WaterDrag;
+
+		// compute torque based on the linear force we apply
+		const FVec3 TorqueFromWaterVelocity = Chaos::FVec3::CrossProduct(WorldCOMToForcePos, ForceFromWaterVelocity);
+
+		OutTotalWorldForce += ForceFromWaterVelocity;
+		OutTotalWorldTorque += TorqueFromWaterVelocity;
+	}
+
+	void ComputeBuoyantForceForShape(const Chaos::FPBDRigidsEvolution& Evolution, const Chaos::FPBDRigidParticleHandle* RigidParticle, const float DeltaSeconds, const float WaterDensity,
+		const Chaos::FVec3& SubmergedCoM, const float SubmergedVol, const Chaos::FVec3& WaterN, Chaos::FVec3& OutWorldBuoyantForce, Chaos::FVec3& OutWorldBuoyantTorque)
+	{
+		SCOPE_CYCLE_COUNTER(STAT_BuoyancyAlgorithms_ComputeBuoyantForceForShape)
+
+		// Get perparticle gravity rule, for figuring out the effective gravity on buoyant objects
+		const Chaos::FPerParticleGravity* PerParticleGravity = &Evolution.GetGravityForces();
+
+		// Figure out the gravity level of the particle
+		const int32 GravityGroupIndex = RigidParticle->GravityGroupIndex();
+		const Chaos::FVec3 GravityAccelVec
+			= PerParticleGravity != nullptr && GravityGroupIndex != INDEX_NONE
+			? (Chaos::FVec3)PerParticleGravity->GetAcceleration(GravityGroupIndex)
+			: Chaos::FVec3::DownVector * 980.f; // Default to "regular" gravity
+								
+		// NOTE: We assume gravity is -Z for perf... If we want to support buoyancy for
+		// weird gravity setups, this is where we'd have to fix it up.
+		const FVec3 GravityDir = FVec3::DownVector;
+		const float GravityAccel = FVec3::DotProduct(GravityDir, GravityAccelVec);
+
+		// Scale submerged volume
+		float ScaledSubmergedVol = SubmergedVol;
+		float TotalVol;
+		ScaleSubmergedVolume(Evolution, RigidParticle, ScaledSubmergedVol, TotalVol);
+
+		// Compute buoyant force
+		//
+		// NOTE: This is easy to compute with Archimedes' principle
+		// https://en.wikipedia.org/wiki/Buoyancy
+		//
+		const float BuoyantForceMagnitude = WaterDensity * ScaledSubmergedVol * GravityAccel;
+		//                       = [ kg / cm^3 ] * [ cm^3 ]    * [cm / s^2]
+		//                       = [ kg * cm / s^2 ]
+		//                       = [ force ]
+
+		// Only proceed if buoyant force isn't vanishingly small
+		if (BuoyantForceMagnitude < SMALL_NUMBER)
+		{
+			return;
+		}
+
+		// Get a generic particle wrapper
+		FConstGenericParticleHandle RigidGeneric(RigidParticle);
+
+		const FVec3 WorldCoM = RigidGeneric->PCom();
+		const FVec3 CoMDiff = SubmergedCoM - WorldCoM;
+
+		// #todo(dmp): should we push water along normal or -gravity
+		OutWorldBuoyantForce = WaterN * BuoyantForceMagnitude;
+				
+		// Compute world buoyant force and torque	
+		// buoyancy is computed at the submerged center of mass		
+		OutWorldBuoyantTorque = Chaos::FVec3::CrossProduct(CoMDiff, OutWorldBuoyantForce);				
+	}
 }
+
diff --git a/Engine/Plugins/Experimental/Buoyancy/Source/Runtime/Private/BuoyancySubsystem.cpp b/Engine/Plugins/Experimental/Buoyancy/Source/Runtime/Private/BuoyancySubsystem.cpp
--- a/Engine/Plugins/Experimental/Buoyancy/Source/Runtime/Private/BuoyancySubsystem.cpp
+++ b/Engine/Plugins/Experimental/Buoyancy/Source/Runtime/Private/BuoyancySubsystem.cpp
@@ -16,6 +16,7 @@
 #include "WaterSubsystem.h"
 #include "WaterBodyManager.h"
 #include "WaterSplineComponent.h"
+#include "BakedShallowWaterSimulationComponent.h"
 #include "Chaos/PhysicsObject.h"
 #include "PBDRigidsSolver.h"
 #include "Chaos/PBDRigidsEvolutionGBF.h"
@@ -25,6 +26,15 @@
 #include "PhysicsProxy/SingleParticlePhysicsProxy.h"
 #include "Chaos/GeometryParticlesfwd.h"
 
+#include "Chaos/ParticleHandle.h"
+#include "Chaos/DebugDrawQueue.h"
+#include "Chaos/CastingUtilities.h"
+#include "Chaos/Utilities.h"
+#include "Chaos/PBDRigidsEvolutionGBF.h"
+#include "Chaos/Collision/CollisionFilter.h"
+#include "Chaos/Collision/CollisionUtil.h"
+#include "Chaos/Sphere.h"
+
 //
 // CVars
 //
@@ -43,6 +53,14 @@ FAutoConsoleVariableRef CVarBuoyancyDebugDraw(TEXT("p.Buoyancy.DebugDraw"), bBuo
 #endif
 
 
+int32 bUseShallowWaterSimulation = 1;
+static FAutoConsoleVariableRef CVarUseShallowWaterSimulation(
+	TEXT("p.Buoyancy.bUseShallowWaterSimulation"),
+	bUseShallowWaterSimulation,
+	TEXT("new buoancy method"),
+	ECVF_Scalability
+);
+
 //
 // Logging
 //
@@ -68,7 +86,6 @@ static FAutoConsoleCommandWithWorld BuoyancyLogMemory(
 	}));
 #endif // WITH_BUOYANCY_MEMORY_TRACKING
 
-
 //
 // Buoyancy Subsystem
 //
@@ -339,18 +356,22 @@ void UBuoyancySubsystem::UpdateSplineData()
 		{
 			// Get the metadata object if there is one, and the spline component
 			UWaterSplineMetadata* SplineMetadata = WaterBodyComponent->GetWaterSplineMetadata();
+			UBakedShallowWaterSimulationComponent *BakedSim = WaterBodyComponent->GetBakedShallowWaterSimulation();
+			
 			if (UWaterSplineComponent* SplineComponent = WaterBodyComponent->GetWaterSpline())
 			{
 				// Copy out water spline data into a shared ptr, to be associated with all
 				// child particles and marshaled to PT.
 				const Chaos::FRigidTransform3 WaterTransform = WaterBodyComponent->GetComponentTransform();
 				const TOptional<FInterpCurveFloat> EmptyOptionalFloat;
+				const TOptional<FShallowWaterSimulationGrid> EmptyOptionalSimGrid;
 				TSharedPtr<FBuoyancyWaterSplineData> WaterSplineData = MakeShared<FBuoyancyWaterSplineData>(
 					WaterTransform,
 					SplineComponent->SplineCurves.Position,
 					WaterBodyComponent->GetWaterBodyType(),
 					SplineMetadata ? SplineMetadata->RiverWidth : EmptyOptionalFloat,
-					SplineMetadata ? SplineMetadata->WaterVelocityScalar : EmptyOptionalFloat
+					SplineMetadata ? SplineMetadata->WaterVelocityScalar : EmptyOptionalFloat,
+					bUseShallowWaterSimulation && WaterBodyComponent->UseBakedShallowWaterSimulationForBuoyancy() ? BakedSim->SimulationData : EmptyOptionalSimGrid
 				);
 
 				// Go over each physics object in each primitive component which was generated
@@ -675,6 +696,7 @@ void FBuoyancySubsystemSimCallback::OnMidPhaseModification_Internal(Chaos::FMidP
 	// Apply buoyant forces resulting from submersions
 	ApplyBuoyantForces(*Evolution);
 
+	// #todo(dmp): need to generate callback info from shallow water interactions
 	// Generate async outputs for callback data
 	GenerateCallbackData();
 
@@ -758,9 +780,30 @@ void FBuoyancySubsystemSimCallback::TrackInteraction(
 		// Always disable midphases with water
 		MidPhase.Disable();
 	}
-
-	float ClosestSplineKey;
-	FVector ClosestPoint;
+	
+	// #todo(dmp): ideally for shallow water, we'd be sampling the body a bit more to tell if anything other than the center is interacting with water
+	float ClosestSplineKey = 0;
+	FVector ClosestPoint = FVector(0, 0, 0);
+	
+	if (WaterSpline.ShouldSampleFromShallowWaterSimulation())
+	{
+		// eval shallow water - closest point is just the current point adjusted upward in Z
+		const FVector ParticlePos = RigidParticle->XCom();
+		FVector WaterVelocity;
+		float WaterHeight;
+		float WaterDepth;
+		WaterSpline.ShallowWaterSimData->SampleShallowWaterSimulationAtPosition(ParticlePos, WaterVelocity, WaterHeight, WaterDepth);
+		ClosestPoint = ParticlePos;
+		ClosestPoint.Z = WaterHeight;
+		
+		// if there is no water here, then return.  This is an approximation that there is water underneath this rigid that might cause interaction
+		// ideally, we'd multisample this since we can lose some water interactions this way
+		if (WaterDepth < 1e-5)
+		{
+			return;
+		}
+	}
+	else
 	{
 		SCOPE_CYCLE_COUNTER(STAT_BuoyancySubsystem_SplineEvaluation)
 
@@ -791,6 +834,7 @@ void FBuoyancySubsystemSimCallback::TrackInteraction(
 
 		// If this particle already has a list of interactions, add to it. Otherwise
 		// make a new one.
+		// #todo(dmp): we should split out the interaction data here to be for shallow water or not
 		BuoyancyParticleData.GetInteractions(*RigidParticle).Add({
 			RigidParticle,
 			WaterParticle,
@@ -813,8 +857,92 @@ void FBuoyancySubsystemSimCallback::ProcessInteractions(Chaos::FPBDRigidsEvoluti
 
 		// Process each interaction
 		for (FBuoyancyInteraction& Interaction : ParticleInteractions)
+		{						
+			// #todo(dmp): we do not want to allow multiple shallow water sims interacting with a single particle for now.
+			// #todo(dmp): correctly handle when a particle can overlap shallow water and other water bodies at the same time
+			// in order to do this, we must unite the integration methods for shallow water and spline based water bodies
+			// and rework the submersion data structures.
+			if (Interaction.WaterSpline.ShouldSampleFromShallowWaterSimulation())
+			{
+				ProcessInteractionShallowWater(Evolution, Interaction);
+			}
+			else
+			{
+				ProcessInteraction(Evolution, Interaction);
+			}
+		}
+	}
+}
+
+void FBuoyancySubsystemSimCallback::ProcessInteractionShallowWater(Chaos::FPBDRigidsEvolution& Evolution, FBuoyancyInteraction& Interaction)
+{		
+	const Chaos::FReal DeltaSeconds = GetDeltaTime_Internal();
+		
+	if (!Interaction.WaterSpline.ShallowWaterSimData.IsSet())
+	{
+		return;
+	}
+
+	FShallowWaterSimulationGrid DefaultShallowWaterSimulationGrid;
+	const FShallowWaterSimulationGrid& ShallowWaterSimData = Interaction.WaterSpline.ShallowWaterSimData.Get(DefaultShallowWaterSimulationGrid);
+
+	float TotalSubmergedVol;
+	Chaos::FVec3 TotalSubmergedCoM;
+	Chaos::FVec3 TotalWorldForce;
+	Chaos::FVec3 TotalWorldTorque;
+		
+	// #todo(dmp): abstract out the ShallowWaterSimData to generic class so we can sample other velocity sources (ie: splines, constant value, etc)
+	BuoyancyAlgorithms::ComputeSubmergedVolumeAndForcesForParticle(BuoyancyParticleData,
+		Interaction.RigidParticle, Interaction.WaterParticle,
+		ShallowWaterSimData,
+		Evolution, DeltaSeconds, BuoyancySettings->WaterDensity, BuoyancySettings->WaterDrag,
+		TotalSubmergedVol, TotalSubmergedCoM, TotalWorldForce, TotalWorldTorque);
+
+	// Apply forces to the particle
+	{
+		SCOPE_CYCLE_COUNTER(STAT_BuoyancySubsystem_ShallowWaterIntegrateForces)
+
+		Chaos::FConstGenericParticleHandle RigidGeneric(Interaction.RigidParticle);
+
+#if ENABLE_DRAW_DEBUG
+		if (bBuoyancyDebugDraw)
+		{
+			Chaos::FDebugDrawQueue::GetInstance().DrawDebugSphere(TotalSubmergedCoM, 20, 10, FColor::Green, false, -1.f, -1, 2.0f);
+			Chaos::FDebugDrawQueue::GetInstance().DrawDebugSphere(RigidGeneric->PCom(), 20, 10, FColor::Blue, false, -1.f, -1, 2.0f);
+
+			Chaos::FDebugDrawQueue::GetInstance().DrawDebugDirectionalArrow(RigidGeneric->PCom(), RigidGeneric->PCom() + TotalWorldForce, 20, FColor::Blue, false, -1.f, -1, 2.f);
+		}
+#endif
+
+		// #todo(dmp): if we do this integration to accelerations per shape, then each shape can have different densities per shape.
+		// This will change to integrate the accumulated accelerations instead of forces
+
+		// Use inertia to convert forces to accelerations
+		const Chaos::FVec3 LinearAccel = RigidGeneric->InvM() * TotalWorldForce;
+
+		const Chaos::FMatrix33 WorldInvI = Chaos::Utilities::ComputeWorldSpaceInertia(RigidGeneric->RCom(), RigidGeneric->ConditionedInvI());
+		const Chaos::FVec3 AngularAccel = WorldInvI * TotalWorldTorque;
+
+		// Integrate to get delta velocities
+		Chaos::FVec3 DeltaV = LinearAccel * DeltaSeconds;
+		Chaos::FVec3 DeltaW = AngularAccel * DeltaSeconds;
+	
+		// Clamp delta velocities
+		DeltaV = DeltaV.GetClampedToSize(0.f, BuoyancySettings->MaxDeltaV);
+		DeltaW = DeltaW.GetClampedToSize(0.f, BuoyancySettings->MaxDeltaW);
+
+		// Compute new velocity
+		const Chaos::FVec3 NewV = Interaction.RigidParticle->GetV() + DeltaV;
+		const Chaos::FVec3 NewW = Interaction.RigidParticle->GetW() + DeltaW;
+
+		// Apply the deltas
+		Interaction.RigidParticle->SetV(NewV);
+		Interaction.RigidParticle->SetW(NewW);
+	
+		// Wake up the body??
+		if (BuoyancySettings->bKeepAwake)
 		{
-			ProcessInteraction(Evolution, Interaction);
+			Evolution.SetParticleObjectState(Interaction.RigidParticle, Chaos::EObjectStateType::Dynamic);
 		}
 	}
 }
diff --git a/Engine/Plugins/Experimental/Buoyancy/Source/Runtime/Public/BuoyancyAlgorithms.h b/Engine/Plugins/Experimental/Buoyancy/Source/Runtime/Public/BuoyancyAlgorithms.h
--- a/Engine/Plugins/Experimental/Buoyancy/Source/Runtime/Public/BuoyancyAlgorithms.h
+++ b/Engine/Plugins/Experimental/Buoyancy/Source/Runtime/Public/BuoyancyAlgorithms.h
@@ -6,6 +6,7 @@
 #include "Chaos/Core.h"
 #include "Chaos/ParticleHandleFwd.h"
 #include "Chaos/AABB.h"
+#include "WaterBodyComponent.h"
 
 namespace Chaos
 {
@@ -63,5 +64,49 @@ namespace BuoyancyAlgorithms
 
 	// Given a rigid particle and it's submerged CoM and Volume, compute delta velocities for
 	// integrated buoyancy forces on an object
-	bool ComputeBuoyantForce(const Chaos::FPBDRigidParticleHandle* RigidParticle, const float DeltaSeconds, const float WaterDensity, const float WaterDrag, const Chaos::FVec3& GravityAccelVec, const Chaos::FVec3& SubmergedCoM, const float SubmergedVol, const Chaos::FVec3& WaterVel, const Chaos::FVec3& WaterN, Chaos::FVec3& OutDeltaV, Chaos::FVec3& OutDeltaW);
+	bool ComputeBuoyantForce(const Chaos::FPBDRigidParticleHandle* RigidParticle, const float DeltaSeconds, const float WaterDensity, 
+		const float WaterDrag, const Chaos::FVec3& GravityAccelVec, const Chaos::FVec3& SubmergedCoM, const float SubmergedVol, 
+		const Chaos::FVec3& WaterVel, const Chaos::FVec3& WaterN, Chaos::FVec3& OutDeltaV, Chaos::FVec3& OutDeltaW);
+
+	// given a particle, loop over the contained shapes and accumulate force/torque/submerged CoM values
+	void ComputeSubmergedVolumeAndForcesForParticle(FBuoyancyParticleData& ParticleData,
+		const Chaos::FGeometryParticleHandle* SubmergedParticle, const Chaos::FGeometryParticleHandle* WaterParticle,
+		const FShallowWaterSimulationGrid& ShallowWaterGrid,
+		const Chaos::FPBDRigidsEvolution& Evolution, const float DeltaSeconds, const float WaterDensity, const float WaterDrag,
+		float& OutTotalSubmergedVol, Chaos::FVec3& OutTotalSubmergedCoM, Chaos::FVec3& OutTotalForce, Chaos::FVec3& OutTotalTorque);
+
+	// given a shape, compute the submerged volume and accumulate forces
+	// this is done in a single function call because of the iterative nature of the algorithm
+	void ComputeSubmergedVolumeAndForcesForShape(
+		const Chaos::FGeometryParticleHandle* SubmergedParticle,
+		const Chaos::FPBDRigidsEvolution& Evolution, float DeltaSeconds, const float WaterDensity, const float WaterDrag,
+		const Chaos::FImplicitObject* Implicit, const Chaos::FRigidTransform3& RelativeTransform, const int32 RootObjectIndex, const int32 ObjectIndex, const int32 LeafObjectIndex,
+		const FShallowWaterSimulationGrid& ShallowWaterGrid, 
+		Chaos::FVec3& OutWaterP, Chaos::FVec3& OutWaterN, 
+		float& OutSubmergedVol, Chaos::FVec3& OutSubmergedCoM,
+		Chaos::FVec3& OutForce, Chaos::FVec3& OutTorque);
+
+	// find intersection points between a plane and aabbox
+	void FindAllIntersectionPoints(const Chaos::FVec3& WaterP, const Chaos::FVec3& WaterN, const Chaos::FAABB3& LocalBox, const FVector WorldVertexPosition[8],
+		TMap<int, FVector> &OutEdgeToIntersectionPointMap, TArray<FVector>& OutOrderedIntersectionPoints, FVector& OutIntersectionCenter);
+
+	// sort intersection points by angle
+	void SortIntersectionPointsByAngle(const Chaos::FVec3& WaterP, const Chaos::FVec3& WaterN, const Chaos::FVec3 &IntersectionCenter, 
+		const TMap<int, FVector>& EdgeToIntersectionPointMap, TArray<FVector>& OutOrderedIntersectionPoints);
+
+	// compute area and volume of a tet from a triangle and center point on mesh
+	void FORCEINLINE ComputeTriangleAreaAndVolume(const FVector &V0, const FVector &V1, const FVector &V2,
+		const FVector &MeshCenterPoint, FVector& OutTriangleBaryCenter, FVector& OutNormal, float& OutArea, float& OutVolume, bool DebugDraw = false);
+
+	// compute the force the fluid exerts on a triangle
+	void FORCEINLINE ComputeFluidForceForTriangle(const float WaterDrag,
+		const float DeltaSeconds,
+		const Chaos::FPBDRigidParticleHandle* RigidParticle, const FVector WorldCoM,
+		const FVector &TriBaryCenter, const FVector &TriNormal, const float TriArea, const float TetVolume,
+		const FVector &WaterVelocity, const FVector &WaterP, const FVector &WaterN,
+		FVector& OutTotalWorldForce, FVector& OutTotalWorldTorque);
+
+	// compute the force the buoyancy exerts on a shape
+	void ComputeBuoyantForceForShape(const Chaos::FPBDRigidsEvolution& Evolution, const Chaos::FPBDRigidParticleHandle* RigidParticle, const float DeltaSeconds, const float WaterDensity,
+		const Chaos::FVec3& SubmergedCoM, const float SubmergedVol, const Chaos::FVec3& WaterN, Chaos::FVec3& OutWorldBuoyantForce, Chaos::FVec3& OutWorldBuoyantTorque);
 }
diff --git a/Engine/Plugins/Experimental/Buoyancy/Source/Runtime/Public/BuoyancyStats.h b/Engine/Plugins/Experimental/Buoyancy/Source/Runtime/Public/BuoyancyStats.h
--- a/Engine/Plugins/Experimental/Buoyancy/Source/Runtime/Public/BuoyancyStats.h
+++ b/Engine/Plugins/Experimental/Buoyancy/Source/Runtime/Public/BuoyancyStats.h
@@ -25,6 +25,8 @@ DECLARE_CYCLE_STAT(TEXT("Buoyancy_Subsystem_Tick"), STAT_BuoyancySubsystem_Tick,
 DECLARE_CYCLE_STAT(TEXT("Buoyancy_Subsystem_UpdateBuoyancySettings"), STAT_BuoyancySubsystem_UpdateBuoyancySettings, STATGROUP_Buoyancy);
 DECLARE_CYCLE_STAT(TEXT("Buoyancy_Subsystem_UpdateWaterBodiesList"), STAT_BuoyancySubsystem_UpdateWaterBodiesList, STATGROUP_Buoyancy);
 DECLARE_CYCLE_STAT(TEXT("Buoyancy_Subsystem_DispatchCallbacks"), STAT_BuoyancySubsystem_DispatchCallbacks, STATGROUP_Buoyancy);
+DECLARE_CYCLE_STAT(TEXT("Buoyancy_Subsystem_ShallowWaterIntegrateForces"), STAT_BuoyancySubsystem_ShallowWaterIntegrateForces, STATGROUP_Buoyancy);
+DECLARE_CYCLE_STAT(TEXT("Buoyancy_Subsystem_ProcessInteraction"), STAT_BuoyancySubsystem_ProcessInteraction, STATGROUP_Buoyancy);
 
 // Algorithms
 DECLARE_CYCLE_STAT(TEXT("Buoyancy_Algorithms_ComputeSubmergedVolume"), STAT_BuoyancyAlgorithms_ComputeSubmergedVolume, STATGROUP_Buoyancy);
@@ -33,6 +35,13 @@ DECLARE_CYCLE_STAT(TEXT("Buoyancy_Algorithms_ScaleSubmergedVolume"), STAT_Buoyan
 DECLARE_CYCLE_STAT(TEXT("Buoyancy_Algorithms_SubdivideBounds"), STAT_BuoyancyAlgorithms_SubdivideBounds, STATGROUP_Buoyancy);
 DECLARE_CYCLE_STAT(TEXT("Buoyancy_Algorithms_ComputeBuoyantForces"), STAT_BuoyancyAlgorithms_ComputeBuoyantForces, STATGROUP_Buoyancy);
 DECLARE_CYCLE_STAT(TEXT("Buoyancy_Algorithms_SubmergeShapeInternal"), STAT_BuoyancyAlgorithms_SubmergeShapeInternal, STATGROUP_Buoyancy);
+DECLARE_CYCLE_STAT(TEXT("Buoyancy_Algorithms_ComputeSubmergedVolumeAndForcesForParticle"), STAT_BuoyancyAlgorithms_ComputeSubmergedVolumeAndForcesForParticle, STATGROUP_Buoyancy);
+DECLARE_CYCLE_STAT(TEXT("Buoyancy_Algorithms_ComputeSubmergedVolumeAndForcesForShape"), STAT_BuoyancyAlgorithms_ComputeSubmergedVolumeAndForcesForShape, STATGROUP_Buoyancy);
+DECLARE_CYCLE_STAT(TEXT("Buoyancy_Algorithms_FindAllIntersectionPoints"), STAT_BuoyancyAlgorithms_FindAllIntersectionPoints, STATGROUP_Buoyancy);
+DECLARE_CYCLE_STAT(TEXT("Buoyancy_Algorithms_SortIntersectionPointsByAngle"), STAT_BuoyancyAlgorithms_SortIntersectionPointsByAngle, STATGROUP_Buoyancy);
+DECLARE_CYCLE_STAT(TEXT("Buoyancy_Algorithms_ComputeTriangleAreaAndVolume"), STAT_BuoyancyAlgorithms_ComputeTriangleAreaAndVolume, STATGROUP_Buoyancy);
+DECLARE_CYCLE_STAT(TEXT("Buoyancy_Algorithms_ComputeFluidForceForTriangle"), STAT_BuoyancyAlgorithms_ComputeFluidForceForTriangle, STATGROUP_Buoyancy);
+DECLARE_CYCLE_STAT(TEXT("Buoyancy_Algorithms_ComputeBuoyantForceForShape"), STAT_BuoyancyAlgorithms_ComputeBuoyantForceForShape, STATGROUP_Buoyancy);
 
 // Particle Data
 DECLARE_CYCLE_STAT(TEXT("Buoyancy_ParticleData_GetData"), STAT_BuoyancyParticleData_GetData, STATGROUP_Buoyancy);
diff --git a/Engine/Plugins/Experimental/Buoyancy/Source/Runtime/Public/BuoyancySubsystem.h b/Engine/Plugins/Experimental/Buoyancy/Source/Runtime/Public/BuoyancySubsystem.h
--- a/Engine/Plugins/Experimental/Buoyancy/Source/Runtime/Public/BuoyancySubsystem.h
+++ b/Engine/Plugins/Experimental/Buoyancy/Source/Runtime/Public/BuoyancySubsystem.h
@@ -229,10 +229,11 @@ private:
 	void TrackInteractions(Chaos::FPBDRigidsSolver& PBDSolver, Chaos::FPBDRigidsEvolution& Evolution, Chaos::FMidPhaseModifierAccessor& MidPhaseAccessor);
 	void TrackInteraction(Chaos::FPBDRigidsEvolution& Evolution, Chaos::FGeometryParticleHandle* WaterParticle, Chaos::FPBDRigidParticleHandle* RigidParticle, const FBuoyancyWaterSplineData& WaterSpline, Chaos::FMidPhaseModifier& MidPhase);
 	void ProcessInteractions(Chaos::FPBDRigidsEvolution& Evolution);
+	void ProcessInteractionShallowWater(Chaos::FPBDRigidsEvolution& Evolution, FBuoyancyInteraction& Interaction);
 	void ProcessInteraction(Chaos::FPBDRigidsEvolution& Evolution, FBuoyancyInteraction& Interaction);
 	void ApplyBuoyantForces(Chaos::FPBDRigidsEvolution& Evolution);
 	void GenerateCallbackData();
-
+	
 #if WITH_BUOYANCY_MEMORY_TRACKING
 	void GenerateAllocationData();
 	SIZE_T AllocatedSize = 0;
diff --git a/Engine/Plugins/Experimental/Buoyancy/Source/Runtime/Public/BuoyancyWaterSplineData.h b/Engine/Plugins/Experimental/Buoyancy/Source/Runtime/Public/BuoyancyWaterSplineData.h
--- a/Engine/Plugins/Experimental/Buoyancy/Source/Runtime/Public/BuoyancyWaterSplineData.h
+++ b/Engine/Plugins/Experimental/Buoyancy/Source/Runtime/Public/BuoyancyWaterSplineData.h
@@ -12,15 +12,23 @@ struct FBuoyancyWaterSplineData
 		const FInterpCurveVector& InPosition,
 		const EWaterBodyType InBodyType,
 		const TOptional<FInterpCurveFloat>& InWidth,
-		const TOptional<FInterpCurveFloat>& InVelocity)
+		const TOptional<FInterpCurveFloat>& InVelocity,
+		const TOptional<FShallowWaterSimulationGrid>& InShallowWaterSimData
+	)
 		: Transform(InTransform)
 		, Position(InPosition)
 		, BodyType(InBodyType)
 		, Width(InWidth)
 		, Velocity(InVelocity)
+		, ShallowWaterSimData(InShallowWaterSimData)
 	{ }
 
 
+	bool ShouldSampleFromShallowWaterSimulation() const
+	{
+		return ShallowWaterSimData.IsSet() && ShallowWaterSimData->IsValid();
+	}
+
 	// Parameters that all water bodies have
 	Chaos::FRigidTransform3 Transform;
 	FInterpCurveVector Position;
@@ -29,4 +37,8 @@ struct FBuoyancyWaterSplineData
 	// Parameters that only _some_ water bodies have
 	TOptional<FInterpCurveFloat> Width;
 	TOptional<FInterpCurveFloat> Velocity;
+
+	// Only some water bodies can have a baked shallow water sim representation, which
+	// overrides the splines for water height/depth/velocity/normal
+	TOptional<FShallowWaterSimulationGrid> ShallowWaterSimData;
 };
\ No newline at end of file
diff --git a/Engine/Source/Runtime/Experimental/Chaos/Public/Chaos/AABB.h b/Engine/Source/Runtime/Experimental/Chaos/Public/Chaos/AABB.h
--- a/Engine/Source/Runtime/Experimental/Chaos/Public/Chaos/AABB.h
+++ b/Engine/Source/Runtime/Experimental/Chaos/Public/Chaos/AABB.h
@@ -25,6 +25,12 @@ namespace Chaos
 		int8 VertexIndex1;
 	};
 
+	struct CHAOS_API FAABBFace
+	{
+		int8 VertexIndex[4];
+		int8 EdgeIndex[4];
+	};
+
 	template<class T, int d>
 	class TAABB
 	{
@@ -526,17 +532,17 @@ namespace Chaos
 		*	This algorithm produces the following index scheme, where the
 		*	bottom left vertex is the "min" vertex and the top right is "max":
 		*	
-		*	        *-----11------*
+		*	        6-----11------7
 		*	       /|            /|
 		*	      9 |          10 |
 		*	     /  6          /  7
-		*	    *------8------*   |
+		*	    4------8------5   |
 		*	    |   |         |   |
-		*	    |   *------5--|---*
+		*	    |   2------5--|---3
 		*	    2  /          4  /
 		*	    | 1           | 3
 		*	    |/            |/
-		*	    *------0------*
+		*	    0------0------1
 		*/
 		FORCEINLINE FAABBEdge GetEdge(const int32 Index) const
 		{
@@ -552,6 +558,22 @@ namespace Chaos
 			return Edges[Index];
 		}
 
+		FORCEINLINE FAABBFace GetFace(const int32 Index) const
+		{
+			// See "GetVertex(int32)"
+			check(0 <= Index && Index < 6);
+			static constexpr FAABBFace Faces[]
+			{
+				{{ 0, 1, 5, 4 }, {0, 4, 8, 2}},
+				{{ 1, 3, 7, 5 }, {3, 7, 10, 4}},
+				{{ 4, 5, 7, 6 }, {8, 10, 11, 9}},
+				{{ 0, 2, 3, 1 }, {1, 5, 3, 0}},
+				{{ 0, 4, 6, 2 }, {2, 9, 6, 1}},
+				{{ 2, 6, 7, 3 }, {6, 11, 7, 5}}
+			};
+			return Faces[Index];
+		}		
+
 		FORCEINLINE int LargestAxis() const
 		{
 			const auto Extents = this->Extents();
