Fix incorrect synctimestamp check for particle validity on BeginContact callback queue

diff --git a/Engine/Plugins/Experimental/Buoyancy/Source/Runtime/Private/BuoyancySubsystem.cpp b/Engine/Plugins/Experimental/Buoyancy/Source/Runtime/Private/BuoyancySubsystem.cpp
--- a/Engine/Plugins/Experimental/Buoyancy/Source/Runtime/Private/BuoyancySubsystem.cpp
+++ b/Engine/Plugins/Experimental/Buoyancy/Source/Runtime/Private/BuoyancySubsystem.cpp
@@ -1267,7 +1267,7 @@ void FBuoyancySubsystemSimCallback::GenerateCallbackData()
 			// Build up output of new and continuing surface touches
 			for (const FBuoyancySubmersionMetaData::FWaterContact& WaterContact : MetaData.WaterContacts)
 			{
-				if (!ensureMsgf(IsParticleValid(WaterContact.Water, Submersion.SyncTimestamp.Pin()), TEXT("Submersion data for buoyancy callback includes invalid water body particle handle")))
+				if (!ensureMsgf(IsParticleValid(WaterContact.Water, WaterContact.SyncTimestamp.Pin()), TEXT("Submersion data for buoyancy callback includes invalid water body particle handle")))
 				{
 					continue;
 				}
