new integration method
  Spline rivers support (required refactor to either sample baked simulations or splines)
  Callback data generation
  Multiple overlapping water bodies
Global cvar to turn on new integration method for all water bodies (off by default)
global cvar to perform buoyancy along the surface normal or -gravity
update example maps and system

#virtualized

diff --git a/Engine/Plugins/Experimental/Buoyancy/Source/Runtime/Private/BuoyancyAlgorithms.cpp b/Engine/Plugins/Experimental/Buoyancy/Source/Runtime/Private/BuoyancyAlgorithms.cpp
--- a/Engine/Plugins/Experimental/Buoyancy/Source/Runtime/Private/BuoyancyAlgorithms.cpp
+++ b/Engine/Plugins/Experimental/Buoyancy/Source/Runtime/Private/BuoyancyAlgorithms.cpp
@@ -13,6 +13,8 @@
 #include "Chaos/Collision/CollisionFilter.h"
 #include "Chaos/Collision/CollisionUtil.h"
 #include "Chaos/Sphere.h"
+#include "Chaos/Convex.h"
+#include "Chaos/ImplicitObjectScaled.h"
 
 //
 // CVars
@@ -23,6 +25,8 @@ extern bool bBuoyancyDebugDraw;
 bool bBuoyancyAlgorithmsAllowVolRatioOverOne = false;
 FAutoConsoleVariableRef CVarBuoyancyAlgorithmsAllowVolRatioOverOne(TEXT("p.Buoyancy.Algorithms.AllowVolRatioOverOne"), bBuoyancyAlgorithmsAllowVolRatioOverOne, TEXT(""));
 
+bool bBuoyancyAlgorithmsBuoyancyAlongNormal = true;
+FAutoConsoleVariableRef CVarBuoyancyAlgorithmsBuoyancyAlongNormal(TEXT("p.Buoyancy.Algorithms.BuoyancyAlongNormal"), bBuoyancyAlgorithmsBuoyancyAlongNormal, TEXT(""));
 
 //
 // Internal Functions
@@ -133,6 +137,9 @@ namespace BuoyancyAlgorithms
 	using Chaos::EImplicitObjectType;
 	using Chaos::FConstGenericParticleHandle;
 	using Chaos::FImplicitSphere3;
+	using Chaos::FConvex;
+	using Chaos::TImplicitObjectScaled;
+	using Chaos::TImplicitObjectInstanced;
 
 	FRealSingle ComputeParticleVolume(const FPBDRigidsEvolutionGBF& Evolution, const FGeometryParticleHandle* Particle)
 	{
@@ -584,7 +591,18 @@ namespace BuoyancyAlgorithms
 		const FMatrix33 WorldInvI = Chaos::Utilities::ComputeWorldSpaceInertia(RigidGeneric->RCom(), RigidGeneric->ConditionedInvI());
 
 		// Compute world buoyant force and torque
-		const FVec3 WorldForce = WaterN * BuoyantForce;
+		FVec3 WorldForce;
+
+		// optionally apply buoyancy along normal or -gravity
+		if (bBuoyancyAlgorithmsBuoyancyAlongNormal)
+		{
+			WorldForce = WaterN * BuoyantForce;
+		}
+		else
+		{
+			WorldForce = -1.0 * GravityDir * BuoyantForce;
+		}
+
 		const FVec3 WorldTorque = Chaos::FVec3::CrossProduct(CoMDiff, WorldForce);
 
 		// Use inertia to convert forces to accelerations
@@ -618,14 +636,15 @@ namespace BuoyancyAlgorithms
 		
 	void ComputeSubmergedVolumeAndForcesForParticle(FBuoyancyParticleData& ParticleData, 
 		const Chaos::FGeometryParticleHandle* SubmergedParticle, const FGeometryParticleHandle* WaterParticle, 
-		const FShallowWaterSimulationGrid& ShallowWaterGrid,
+		TSharedPtr<FBuoyancyWaterSampler> WaterSampler,
 		const Chaos::FPBDRigidsEvolution& Evolution, const float DeltaSeconds, const float WaterDensity, const float WaterDrag,
-		float &OutTotalSubmergedVol, Chaos::FVec3 &OutTotalSubmergedCoM, Chaos::FVec3 &OutTotalForce, Chaos::FVec3 &OutTotalTorque)
+		float &OutTotalParticleVol, float &OutTotalSubmergedVol, Chaos::FVec3 &OutTotalSubmergedCoM, Chaos::FVec3 &OutTotalForce, Chaos::FVec3 &OutTotalTorque)
 	{		
 		SCOPE_CYCLE_COUNTER(STAT_BuoyancyAlgorithms_ComputeSubmergedVolumeAndForcesForParticle)
 
 		// Initialize cumulative values
-		OutTotalSubmergedVol = 0;
+		OutTotalParticleVol = 0.f;
+		OutTotalSubmergedVol = 0.f;
 		OutTotalSubmergedCoM = FVec3::ZeroVector;
 		OutTotalForce = FVec3::ZeroVector;
 		OutTotalTorque = FVec3::ZeroVector;
@@ -644,15 +663,16 @@ namespace BuoyancyAlgorithms
 		const FShapeInstanceArray& WaterShapeInstances = WaterParticle->ShapeInstances();
 		const FShapeInstance* WaterShapeInstance = WaterShapeInstances[0].Get();
 
+		WaterSampler->InitializeForSubmersion();
+
 		// Traverse the submerged particle's leaves
 		RootImplicit->VisitLeafObjects(
-			[&Evolution, DeltaSeconds, WaterDensity, WaterDrag, SubmergedParticle, &ShallowWaterGrid, 
+			[&Evolution, DeltaSeconds, WaterDensity, WaterDrag, SubmergedParticle, &WaterSampler, 
 			ParticleIndex, &ShapeInstances, WaterShapeType, WaterShapeInstance, &ParticleWorldTransform, 
 			&SubmergedShapes, &OutTotalSubmergedVol, &OutTotalSubmergedCoM,
 			&OutTotalForce, &OutTotalTorque]
 			(const FImplicitObject* Implicit, const FRigidTransform3& RelativeTransform, const int32 RootObjectIndex, const int32 ObjectIndex, const int32 LeafObjectIndex)
-			{				
-				const FAABB3 RelativeBounds = Implicit->CalculateTransformedBounds(RelativeTransform);
+			{								
 				const int32 ShapeIndex = (ShapeInstances.IsValidIndex(RootObjectIndex)) ? RootObjectIndex : 0;
 				const FShapeInstance* ShapeInstance = ShapeInstances[ShapeIndex].Get();
 				const EImplicitObjectType ShapeType = Chaos::Private::GetImplicitCollisionType(SubmergedParticle, Implicit);
@@ -677,7 +697,7 @@ namespace BuoyancyAlgorithms
 				ComputeSubmergedVolumeAndForcesForShape(SubmergedParticle,
 					Evolution, DeltaSeconds, WaterDensity, WaterDrag,
 					Implicit, RelativeTransform, RootObjectIndex, ObjectIndex, LeafObjectIndex,
-					ShallowWaterGrid,
+					WaterSampler,
 					OutWaterP, OutWaterN,
 					OutSubmergedVol, OutSubmergedCoM,
 					OutForce, OutTorque);
@@ -697,12 +717,15 @@ namespace BuoyancyAlgorithms
 
 		// compute final force and torque for particle as weighted average of all the shapes
 		OutTotalSubmergedCoM /= OutTotalSubmergedVol;
+
+		float ScaledSubmergedVol = OutTotalSubmergedVol;
+		ScaleSubmergedVolume(Evolution, SubmergedParticle, ScaledSubmergedVol, OutTotalParticleVol);
 	}
 
 	void ComputeSubmergedVolumeAndForcesForShape(const Chaos::FGeometryParticleHandle* SubmergedParticle,
 		const Chaos::FPBDRigidsEvolution& Evolution, float DeltaSeconds, const float WaterDensity, const float WaterDrag,
 		const Chaos::FImplicitObject* Implicit, const Chaos::FRigidTransform3& RelativeTransform, const int32 RootObjectIndex, const int32 ObjectIndex, const int32 LeafObjectIndex,
-		const FShallowWaterSimulationGrid& ShallowWaterGrid,		
+		TSharedPtr<FBuoyancyWaterSampler> WaterSampler,
 		Chaos::FVec3& OutWaterP, Chaos::FVec3& OutWaterN,
 		float& OutSubmergedVol, Chaos::FVec3& OutSubmergedCoM,
 		Chaos::FVec3& OutForce, Chaos::FVec3& OutTorque)
@@ -727,47 +750,59 @@ namespace BuoyancyAlgorithms
 			ShapeWorldTransform.SetTranslation(SphereCenter);
 
 			// #todo(dmp): at this point, we haven't figured out the water normal so we use up vector for now
-			FVec3 TmpWaterN = {0, 0, 1};
+			FVec3 TmpWaterN = {0.f, 0.f, 1.f };
 			ShapeWorldTransform.SetRotation(FRotationMatrix::MakeFromZ(TmpWaterN).ToQuat());
 		}
-				
+		
 		// water plane parameters for the body.  Note we have 1 plane per box
-		OutWaterP = FVec3(0, 0, 0);
-		OutWaterN = FVec3(0, 0, 1);
+		OutWaterP = FVec3::ZeroVector;
+		OutWaterN = FVec3(0.f, 0.f, 1.f);
 
 		// Compute submerged volume and CoM
-		OutSubmergedVol = 0;
-		OutSubmergedCoM = FVec3(0, 0, 0);
+		OutSubmergedVol = 0.f;
+		OutSubmergedCoM = FVec3::ZeroVector;
 
 		// output forces
-		OutForce = FVec3(0, 0, 0);
-		OutTorque = FVec3(0, 0, 0);
+		OutForce = FVec3::ZeroVector;
+		OutTorque = FVec3::ZeroVector;
 
-		float SubmergedCOMTotalWeight = 0;
-		
-		FVector WorldBoxCenter = ShapeWorldTransform.TransformPosition(LocalBox.GetCenter());
+		float SubmergedCOMTotalWeight = 0.f;
+				
+		// either use a water plane per shape's centroid or the closest water point to the particle
+		if (WaterSampler->UsePerShapeWaterPlane())
+		{			
+			const FVector WorldBoxCenter = ShapeWorldTransform.TransformPosition(LocalBox.GetCenter());
 
-		// #todo(dmp): just store normals in the shallow water texture and avoid all this
-		// water plane position is computed from the centroid of the box		
-		FVector CenterVelocity;
-		float CenterHeight;
-		float CenterDepth;
-		ShallowWaterGrid.SampleShallowWaterSimulationAtPosition(WorldBoxCenter, CenterVelocity, CenterHeight, CenterDepth);
+			// water plane position is computed from the centroid of the box		
+			FVector CenterVelocity;
+			float CenterHeight;
+			float CenterDepth;
+			WaterSampler->SampleWaterAtPosition(WorldBoxCenter, CenterVelocity, CenterHeight, CenterDepth);
 
-		// no water at shape centroid of shape so we return
-		if (CenterDepth < SMALL_NUMBER)
+			// no water at shape centroid of shape so we return
+			if (CenterDepth < SMALL_NUMBER)
+			{
+				return;
+			}
+						
+			OutWaterP = WorldBoxCenter;
+			OutWaterP.Z = CenterHeight;
+		}
+		else
 		{
-			return;
+			OutWaterP = WaterSampler->GetClosestWaterPointToParticle();
 		}
 
-		// #todo(dmp): lift normal computation out of this to make function more generic
-		// water plane is based on the shape's centroid
-		OutWaterP = WorldBoxCenter;
-		OutWaterP.Z = CenterHeight;
-		OutWaterN = ShallowWaterGrid.ComputeShallowWaterSimulationNormalAtPosition(OutWaterP);
+		// Sample normal at water plane position
+		// #todo(dmp): just store normals in the shallow water texture rather than finite differencing here
+		OutWaterN = WaterSampler->SampleNormalAtPosition(OutWaterP);
 
 		// #todo(dmp): generalize to convex shapes and remove box assumption
-		const int TotalNumVertices = 8;
+		static constexpr int32 TotalNumVertices = 8;
+		static constexpr int32 NumFaces = 6;
+		static constexpr int32 NumEdges = 12;
+		static constexpr int32 NumVerticesForFace = 4;
+		static constexpr int32 MaxBoxPlaneIntersectionsPoints = 6;
 
 		FVector WorldVertexPosition[TotalNumVertices];
 		bool VertexIsUnderwater[TotalNumVertices];
@@ -775,21 +810,23 @@ namespace BuoyancyAlgorithms
 		// keep track of a reference point inside the clipped geometry that is the average of all the 
 		// submerged vertices and intersection points this will jump around because it isn't weighted by the area/volume
 		// but is sufficient for tet volume computations since we only need an interior point
-		FVector InteriorRefPoint(0, 0, 0);
-		int InteriorRefPointCount = 0;
-
+		FVector InteriorRefPoint = FVector::ZeroVector;
+		int32 InteriorRefPointCount = 0;
+		
 		// cache water values for each vertex and determine broadphase of shape interacting with the water	
-		for (int i = 0; i < TotalNumVertices; ++i)
+		for (int32 i = 0; i < TotalNumVertices; ++i)
 		{
+			// SCOPE_CYCLE_COUNTER(STAT_BuoyancyAlgorithms_ComputeSubmergedVolumeAndForcesForShape_ComputeWorld);
+
 			// compute world space position of vertex
 			const FVector CurrWorldVertexPosition = ShapeWorldTransform.TransformPosition(LocalBox.GetVertex(i));
 			WorldVertexPosition[i] = CurrWorldVertexPosition;
 
-			if ((CurrWorldVertexPosition - OutWaterP).Dot(OutWaterN) < 0)
+			if ((CurrWorldVertexPosition - OutWaterP).Dot(OutWaterN) < SMALL_NUMBER)
 			{
 				VertexIsUnderwater[i] = true;
 				InteriorRefPoint += WorldVertexPosition[i];
-				InteriorRefPointCount++;
+				InteriorRefPointCount++;				
 			}
 			else
 			{
@@ -797,12 +834,22 @@ namespace BuoyancyAlgorithms
 			}
 		}
 
+		// Early out if no vertices are underwater
+		if (InteriorRefPointCount == 0)
+		{
+			return;
+		}
+
 		// find intersection points of the box and water plane
 		FVector IntersectionCenter;			
-		TArray<FVector> OrderedIntersectionPoints;
-		TMap<int, FVector> AllIntersectionPoints;		
-		BuoyancyAlgorithms::FindAllIntersectionPoints(OutWaterP, OutWaterN, LocalBox, WorldVertexPosition, AllIntersectionPoints, OrderedIntersectionPoints, IntersectionCenter);
-		const int NumIntersectionPoints = AllIntersectionPoints.Num();
+
+		// we can have at most 6 intersection points with a box and a plane			
+		TArray<FVector, TFixedAllocator<MaxBoxPlaneIntersectionsPoints>> OrderedIntersectionPoints;
+
+		FVector EdgeToIntersectionPoint[NumEdges];
+		int32 NumIntersectionPoints;
+		BuoyancyAlgorithms::FindAllIntersectionPoints(OutWaterP, OutWaterN, LocalBox, WorldVertexPosition, EdgeToIntersectionPoint, NumIntersectionPoints, 
+			OrderedIntersectionPoints, IntersectionCenter);
 		
 		// Add intersection points to interior reference point
 		// values are already accumulated into the intersection center
@@ -815,64 +862,55 @@ namespace BuoyancyAlgorithms
 		const FVec3 WorldCoM = SubmergedGeneric->PCom();
 
 		// Sum up submerged areas and volumes for each face
-		const int NumFaces = 6;
-		for (int FaceIdx = 0; FaceIdx < NumFaces; ++FaceIdx)
+		
+		for (int32 FaceIdx = 0; FaceIdx < NumFaces; ++FaceIdx)
 		{
-			const Chaos::FAABBFace CurrFace = LocalBox.GetFace(FaceIdx);
+			// SCOPE_CYCLE_COUNTER(STAT_BuoyancyAlgorithms_ComputeSubmergedVolumeAndForcesForShape_FaceIteration);
 
-			// compute area of unsubmerged face	
-			// #todo(dmp): precompute edge lengths
-			const FVector Pos0 = WorldVertexPosition[CurrFace.VertexIndex[0]];
-			const FVector Pos1 = WorldVertexPosition[CurrFace.VertexIndex[1]];
-			const FVector Pos2 = WorldVertexPosition[CurrFace.VertexIndex[2]];
-			const float FaceArea = (Pos1 - Pos0).Length() * (Pos2 - Pos1).Length();
-
-			// skip small faces from degenerate bodies
-			if (FaceArea < SMALL_NUMBER)
-			{
-				continue;
-			}
+			const Chaos::FAABBFace CurrFace = LocalBox.GetFace(FaceIdx);
 
-			FVector IntersectedFaceCenter(0, 0, 0);
+			FVector IntersectedFaceCenter = FVector::ZeroVector;
 
 			// we know that a given face can have at most 5 points below the surface
 			// #todo(dmp): we need to remove this assumption if we support convex hulls or other types of geometry
 			static constexpr int32 MaxNumVerticesBelowWater = 5;
-			TArray<FVector, TInlineAllocator<MaxNumVerticesBelowWater>> FaceIntersectionPoints;
+			TArray<FVector, TFixedAllocator<MaxNumVerticesBelowWater>> FaceIntersectionPoints;
 
 			// walk vertices for the current face in counter clockwise order and find intersections and
-			// submerged vertices
-			const int NumVerticesForFace = 4;
-			for (int i = 0; i < NumVerticesForFace; ++i)
+			// submerged vertices			
+			for (int32 i = 0; i < NumVerticesForFace; ++i)
 			{
+				// SCOPE_CYCLE_COUNTER(STAT_BuoyancyAlgorithms_ComputeSubmergedVolumeAndForcesForShape_FindFaceIntersections);
+
 				// get the current edge index belonging to the i and i+1 vertices
 				// note this is a bit weird...we could just ask the edge for the vertices but they
 				// might be in the wrong order for correctly constructing the intersected face
-				const int CurrEdgeIdx = CurrFace.EdgeIndex[i];
-
-				// get the indices of the two vertices for the current edge
-				const int Edge0Idx = CurrFace.VertexIndex[i];
-				const int Edge1Idx = CurrFace.VertexIndex[(i + 1) % 4];
-				const FVector V0 = WorldVertexPosition[Edge0Idx];
-				const FVector V1 = WorldVertexPosition[Edge1Idx];
-
+				const int32 CurrEdgeIdx = CurrFace.EdgeIndex[i];
+			
 				// if there is an intersection for this edge, add it to the list for the current face
-				const FVector* EdgeIntersection = AllIntersectionPoints.Find(CurrEdgeIdx);
-				if (EdgeIntersection)
+				const FVector &EdgeIntersection = EdgeToIntersectionPoint[CurrEdgeIdx];
+
+				// #todo(dmp): use bit mask for this
+				if (EdgeIntersection.X < TNumericLimits<float>::Max())
 				{
-					FaceIntersectionPoints.Add(*EdgeIntersection);
-					IntersectedFaceCenter += *EdgeIntersection;
+					FaceIntersectionPoints.Add(EdgeIntersection);
+					IntersectedFaceCenter += EdgeIntersection;
 				}
 
+				// Endpoint of the edge
+				const int32 Edge1Idx = CurrFace.VertexIndex[(i + 1) % 4];
+
 				// add the second endpoint if it is underwater
 				if (VertexIsUnderwater[Edge1Idx])
 				{
+					const FVector& V1 = WorldVertexPosition[Edge1Idx];
+
 					FaceIntersectionPoints.Add(V1);
 					IntersectedFaceCenter += V1;
 				}
 			}
 			
-			const int NumFaceIntersectionPts = FaceIntersectionPoints.Num();
+			const int32 NumFaceIntersectionPts = FaceIntersectionPoints.Num();
 
 			// if the face intersects the water or is fully submerged, compute face coverage
 			if (NumFaceIntersectionPts > 0)
@@ -881,11 +919,12 @@ namespace BuoyancyAlgorithms
 
 				// we have a list of the intersection points for a face in correct counterclockwise order
 				// fan layout for points on face, sum up areas and volume for the current face
-				for (int i = 0; i < NumFaceIntersectionPts; i++)
+				for (int32 i = 0; i < NumFaceIntersectionPts; i++)
 				{
-					const FVector V0 = FaceIntersectionPoints[i];
-					const FVector V1 = FaceIntersectionPoints[(i + 1) % NumFaceIntersectionPts];
-					const FVector V2 = IntersectedFaceCenter;
+					// SCOPE_CYCLE_COUNTER(STAT_BuoyancyAlgorithms_ComputeSubmergedVolumeAndForcesForShape_ComputeFaceForces);
+					const FVector &V0 = FaceIntersectionPoints[i];
+					const FVector &V1 = FaceIntersectionPoints[(i + 1) % NumFaceIntersectionPts];
+					const FVector &V2 = IntersectedFaceCenter;
 
 					FVector TriBaryCenter;
 					float TriArea;
@@ -910,10 +949,12 @@ namespace BuoyancyAlgorithms
 					// sample velocity at center of triangle
 					// #todo(dmp): this should just be interpolated from vertex velocities to avoid so many samples	
 					// #todo(dmp): split arrays out to store velocity separate from other values
-					FVector CurrWaterVelocity;
-					float TmpWaterVal;
-					ShallowWaterGrid.SampleShallowWaterSimulationAtPosition(TriBaryCenter,
-						CurrWaterVelocity, TmpWaterVal, TmpWaterVal);
+					FVector CurrWaterVelocity;					
+					{
+						// SCOPE_CYCLE_COUNTER(STAT_BuoyancyAlgorithms_ComputeSubmergedVolumeAndForcesForShape_SampleVelocity)
+					
+						CurrWaterVelocity = WaterSampler->SampleVelocityAtPosition(TriBaryCenter);
+					}
 
 					FVector TotalWorldForce;
 					FVector TotalWorldTorque;
@@ -941,13 +982,12 @@ namespace BuoyancyAlgorithms
 			}
 		}
 
-		// Sum up tet volumes from the fan of triangles across the intersection face
-		int NumOrderedIntersectionPoints = OrderedIntersectionPoints.Num();
-		for (int i = 0; i < NumOrderedIntersectionPoints; i++)
+		// Sum up tet volumes from the fan of triangles across the intersection face		
+		for (int32 i = 0; i < NumIntersectionPoints; i++)
 		{
-			FVector V0 = OrderedIntersectionPoints[i];
-			FVector V1 = OrderedIntersectionPoints[(i + 1) % NumOrderedIntersectionPoints];
-			FVector V2 = IntersectionCenter;
+			const FVector &V0 = OrderedIntersectionPoints[i];
+			const FVector &V1 = OrderedIntersectionPoints[(i + 1) % NumIntersectionPoints];
+			const FVector &V2 = IntersectionCenter;
 
 			FVector TriBaryCenter;
 			float TriArea;
@@ -955,28 +995,20 @@ namespace BuoyancyAlgorithms
 			FVector TriNormal;
 			BuoyancyAlgorithms::ComputeTriangleAreaAndVolume(V0, V1, V2, InteriorRefPoint, TriBaryCenter, TriNormal, TriArea, TetVolume, bBuoyancyDebugDraw);
 
-			// skip degenerate triangles with 0 area
-			if (TriArea < SMALL_NUMBER)
-			{
-				continue;
-			}
-
 			OutSubmergedCoM += TriBaryCenter * TriArea;
 			SubmergedCOMTotalWeight += TriArea;
 			OutSubmergedVol += TetVolume;
 		}
 
 		// normalize weighted average of contributions to world submerged COM
-		// return if nothing has contributed to CoM
-		if (SubmergedCOMTotalWeight > SMALL_NUMBER)
-		{
-			OutSubmergedCoM /= SubmergedCOMTotalWeight;
-		}
-		else
+		// return if nothing has contributed to CoM	
+		if (SubmergedCOMTotalWeight < SMALL_NUMBER)
 		{
 			return;
 		}
-
+					
+		OutSubmergedCoM /= SubmergedCOMTotalWeight;
+				
 		// #todo(dmp): is it correct to normalize by the accumulated area over which the force is done?  Seems wrong...
 		if (AccumulatedForceWeight > SMALL_NUMBER)
 		{
@@ -996,6 +1028,8 @@ namespace BuoyancyAlgorithms
 #if ENABLE_DRAW_DEBUG
 		if (bBuoyancyDebugDraw)
 		{
+			FVector WorldBoxCenter = ShapeWorldTransform.TransformPosition(LocalBox.GetCenter());
+
 			Chaos::FDebugDrawQueue::GetInstance().DrawDebugSphere(WorldBoxCenter, 10, 10, FColor::Green, false, -1.f, -1, 2.0f);
 			Chaos::FDebugDrawQueue::GetInstance().DrawDebugSphere(InteriorRefPoint, 10, 10, FColor::Yellow, false, -1.f, -1, 2.0f);
 			Chaos::FDebugDrawQueue::GetInstance().DrawDebugSphere(OutWaterP, 10, 10, FColor::Orange, false, -1.f, -1, 2.0f);
@@ -1017,82 +1051,83 @@ namespace BuoyancyAlgorithms
 				FColor::Red, false, -1.f, SDPG_Foreground, 3.f);
 		}
 #endif	
-
 	}
 
 	void FindAllIntersectionPoints(const Chaos::FVec3& WaterP, const Chaos::FVec3& WaterN, const Chaos::FAABB3& LocalBox, const FVector WorldVertexPosition[8],
-		TMap<int, FVector>& OutEdgeToIntersectionPointMap, TArray<FVector>& OutOrderedIntersectionPoints, FVector &OutIntersectionCenter)
+		FVector EdgeToIntersectionPoint[12], int32 &NumIntersections, TArray<FVector, TFixedAllocator<6>>& OutOrderedIntersectionPoints, FVector& OutIntersectionCenter)
 	{		
-		SCOPE_CYCLE_COUNTER(STAT_BuoyancyAlgorithms_FindAllIntersectionPoints)
+		// SCOPE_CYCLE_COUNTER(STAT_BuoyancyAlgorithms_FindAllIntersectionPoints)
 
-		OutIntersectionCenter = FVector(0, 0, 0);
+		OutIntersectionCenter = FVector::ZeroVector;
+		
+		NumIntersections = 0;
 
 		// loop over all edges and compute intersection points
-		for (int CurrEdgeIdx = 0; CurrEdgeIdx < 12; ++CurrEdgeIdx)
+		for (int32 CurrEdgeIdx = 0; CurrEdgeIdx < 12; ++CurrEdgeIdx)
 		{
 			const FAABBEdge CurrEdge = LocalBox.GetEdge(CurrEdgeIdx);
 
-			// edge plane intersection for each edge in counterclockwise order
-			const FVector RayOrigin = WorldVertexPosition[CurrEdge.VertexIndex0];
-			const FVector RayEndPoint = WorldVertexPosition[CurrEdge.VertexIndex1];			
-			const FVector RayDir = RayEndPoint - RayOrigin;
-
-			// if the water plane intersects the current edge, add intersection point it to the list
-			const float DirDotN = RayDir.Dot(WaterN);
-			if (FMath::Abs(DirDotN) > SMALL_NUMBER)
+			Chaos::FVec3 IntersectionPoint;
+			if (EdgePlaneIntersection(WaterP, WaterN, WorldVertexPosition[CurrEdge.VertexIndex0], WorldVertexPosition[CurrEdge.VertexIndex1], IntersectionPoint))
 			{
-				const float t = (WaterP - RayOrigin).Dot(WaterN) / DirDotN;
-				if (t >= 0 && t <= 1)
-				{
-					const FVector IntersectionPoint = RayOrigin + RayDir * t;
-
-					OutEdgeToIntersectionPointMap.Add(CurrEdgeIdx, IntersectionPoint);
-
-					OutIntersectionCenter += IntersectionPoint;
-				}
+				OutIntersectionCenter += IntersectionPoint;
+				NumIntersections++;
 			}
+
+			EdgeToIntersectionPoint[CurrEdgeIdx] = IntersectionPoint;
 		}
 
-		if (OutEdgeToIntersectionPointMap.Num() > 0)
+		if (NumIntersections > 0)
 		{
-			OutIntersectionCenter /= OutEdgeToIntersectionPointMap.Num();
+			OutIntersectionCenter /= NumIntersections;
 
 			// Sort intersection points by angle around reference point
-			BuoyancyAlgorithms::SortIntersectionPointsByAngle(WaterP, WaterN, OutIntersectionCenter, OutEdgeToIntersectionPointMap, OutOrderedIntersectionPoints);
+			BuoyancyAlgorithms::SortIntersectionPointsByAngle(WaterP, WaterN, OutIntersectionCenter, EdgeToIntersectionPoint, NumIntersections, OutOrderedIntersectionPoints);
 		}
 	}
 
-	void SortIntersectionPointsByAngle(const Chaos::FVec3& WaterP, const Chaos::FVec3& WaterN, const Chaos::FVec3& IntersectionCenter, const TMap<int, FVector>& EdgeToIntersectionPointMap, TArray<FVector>& OutOrderedIntersectionPoints)
+	void SortIntersectionPointsByAngle(const Chaos::FVec3& WaterP, const Chaos::FVec3& WaterN, const Chaos::FVec3& IntersectionCenter, 
+		const FVector EdgeToIntersectionPoint[12], const int32 NumIntersections, TArray<FVector, TFixedAllocator<6>>& OutOrderedIntersectionPoints)
 	{
 		SCOPE_CYCLE_COUNTER(STAT_BuoyancyAlgorithms_SortIntersectionPointsByAngle)
 
 		// Order intersection points around intersection center
 		TMap<float, FVector> IntersectionPointAngleToCenter;
+				
+		bool FoundFirst = false;
 
-		auto IntersectionPointMapIt = EdgeToIntersectionPointMap.CreateConstIterator();
-		
 		// first element is assigned to be angle 0
-		FVector CurrPoint = IntersectionPointMapIt.Value();
-		FVector CurrVector = CurrPoint - IntersectionCenter;
-		CurrVector.Normalize();
-
-		const FVector AngleStart = CurrVector;
-		IntersectionPointAngleToCenter.Add(0, CurrPoint);
-
-		// loop over remaining points and compute angles to the first point
-		++IntersectionPointMapIt;
-		for (; IntersectionPointMapIt; ++IntersectionPointMapIt)
+		FVector AngleStart;		
+		
+		// loop over all points, determine angle to reference point
+		for (int32 i = 0; i < 12; ++i)
 		{
-			CurrPoint = IntersectionPointMapIt.Value();
+			float CurrAngle;
+			const FVector &CurrPoint = EdgeToIntersectionPoint[i];
 			
-			CurrVector = CurrPoint - IntersectionCenter;
-			CurrVector.Normalize();
-			
-			float CurrAngle = FMath::Acos(AngleStart.Dot(CurrVector));
+			if (CurrPoint.X >= TNumericLimits<float>::Max()) // not a valid point
+			{
+				continue;
+			}						
+			else if (!FoundFirst) // first point found is the 0 angle point
+			{
+				AngleStart = CurrPoint - IntersectionCenter;
+				AngleStart.Normalize();
 
-			if (AngleStart.Cross(CurrVector).Dot(WaterN) < 0)
+				CurrAngle = 0.f;
+				FoundFirst = true;
+			}
+			else // compute angle from the first point to this one
 			{
-				CurrAngle = 2.f * PI - CurrAngle;
+				FVector CurrVector = CurrPoint - IntersectionCenter;
+				CurrVector.Normalize();
+
+				CurrAngle = FMath::Acos(AngleStart.Dot(CurrVector));
+
+				if (AngleStart.Cross(CurrVector).Dot(WaterN) < SMALL_NUMBER)
+				{
+					CurrAngle = 2.f * PI - CurrAngle;
+				}
 			}
 
 			IntersectionPointAngleToCenter.Add(CurrAngle, CurrPoint);
@@ -1104,21 +1139,45 @@ namespace BuoyancyAlgorithms
 		IntersectionPointAngleToCenter.KeySort([](const float& A, const float& B) {
 			return A < B;
 			});
-
+		
 		// Generate an array of all intersection points		
 		IntersectionPointAngleToCenter.GenerateValueArray(OutOrderedIntersectionPoints);
 	}
 
+	bool EdgePlaneIntersection(const Chaos::FVec3& WaterP, const Chaos::FVec3& WaterN, const Chaos::FVec3& V0, const Chaos::FVec3& V1, Chaos::FVec3& IntersectionPoint)
+	{
+		// #todo(dmp): we currently use a fixed array with sentinel values for edges that have no intersection point
+		// we could use a bit mask or sparse array instead.  We switched from using a map due to perf
+		IntersectionPoint = FVector(TNumericLimits<float>::Max());
+
+		// edge plane intersection for each edge in counterclockwise order
+		const FVector& RayOrigin = V0;
+		const FVector RayDir = V1 - V0;
+
+		// if the water plane intersects the current edge, add intersection point it to the list
+		const float DirDotN = RayDir.Dot(WaterN);
+
+		if (FMath::Abs(DirDotN) > SMALL_NUMBER)
+		{
+			const float t = (WaterP - RayOrigin).Dot(WaterN) / DirDotN;
+			if (t >= 0 && t <= 1)
+			{
+				IntersectionPoint = RayOrigin + RayDir * t;
+
+				return true;
+			}
+		}
+
+		return false;
+	}
+
 	void ComputeTriangleAreaAndVolume(const FVector &V0, const FVector &V1, const FVector &V2, const FVector &MeshCenterPoint, FVector &OutTriangleBaryCenter, FVector &OutNormal, float& OutArea, float& OutVolume, bool DebugDraw /*= false*/)
 	{
-		SCOPE_CYCLE_COUNTER(STAT_BuoyancyAlgorithms_ComputeTriangleAreaAndVolume)
+		// SCOPE_CYCLE_COUNTER(STAT_BuoyancyAlgorithms_ComputeTriangleAreaAndVolume)
 
 		// compute center of triangle
 		OutTriangleBaryCenter = (V0 + V1 + V2) / 3.;
-
-		// add up the area of this triangle for all the faces		 
-		OutArea = .5 * (V2 - V0).Cross(V1 - V0).Length();
-
+		
 		// add up the volume of the tet created by this triangle and the center of the box
 		const FVector A = V0 - MeshCenterPoint;
 		const FVector B = V1 - MeshCenterPoint;
@@ -1128,8 +1187,14 @@ namespace BuoyancyAlgorithms
 		OutVolume = FMath::Abs((N.Dot(C)) / 6.f);
 
 		OutNormal = (V1 - V0).Cross(V2 - V0);
-		OutNormal.Normalize();
+		float NormalLength = OutNormal.Length();
+		
+		// #todo(dmp): careful w/ divide by zero here?
+		OutNormal /= NormalLength;
 
+		// area of the triangle
+		OutArea = .5 * NormalLength;
+		
 #if ENABLE_DRAW_DEBUG
 		if (DebugDraw)
 		{
@@ -1151,7 +1216,7 @@ namespace BuoyancyAlgorithms
 		const FVector &WaterVelocity, const FVector &WaterP, const FVector &WaterN,
 		FVector &OutTotalWorldForce, FVector &OutTotalWorldTorque)
 	{
-		SCOPE_CYCLE_COUNTER(STAT_BuoyancyAlgorithms_ComputeFluidForceForTriangle)
+		// SCOPE_CYCLE_COUNTER(STAT_BuoyancyAlgorithms_ComputeFluidForceForTriangle)
 
 		OutTotalWorldForce = FVector(0, 0, 0);
 		OutTotalWorldTorque = FVector(0, 0, 0);
@@ -1160,7 +1225,7 @@ namespace BuoyancyAlgorithms
 		const FVec3 WorldCOMToForcePos = WorldForcePosition - WorldCoM;
 
 		// project velocity sample onto water plane to support waterfalls and flowing rivers more accurately
-		const FVector WaterVelocityOnPlane = WaterVelocity - WaterVelocity.Dot(WaterN) * WaterN;
+		const FVec3 WaterVelocityOnPlane = WaterVelocity - WaterVelocity.Dot(WaterN) * WaterN;
 
 		// Get world space particle linear velocity at current point
 		// note we are including the linear velocity from torque so objects spin properly in flow
@@ -1184,9 +1249,7 @@ namespace BuoyancyAlgorithms
 		// we reject angles > 90 since that would mean we are hitting a backface and no force should be applied					
 		// default to 1 if we have a still pool of water, we still want this to have a valid response so we don't have
 		// a discontinuity.
-		float FacingTest = 1.f;
-
-		FacingTest = RelativeVelocity.Dot(-TriNormal) / RelativeVelocityMag;
+		float FacingTest = RelativeVelocity.Dot(-TriNormal) / RelativeVelocityMag;
 		FacingTest = FacingTest > 0.f ? FacingTest : 0.f;
 
 		// apply some drag based on facing ratio...could make this more complex
@@ -1217,7 +1280,7 @@ namespace BuoyancyAlgorithms
 		const int32 GravityGroupIndex = RigidParticle->GravityGroupIndex();
 		const Chaos::FVec3 GravityAccelVec
 			= PerParticleGravity != nullptr && GravityGroupIndex != INDEX_NONE
-			? (Chaos::FVec3)PerParticleGravity->GetAcceleration(GravityGroupIndex)
+			? (Chaos::FVec3) PerParticleGravity->GetAcceleration(GravityGroupIndex)
 			: Chaos::FVec3::DownVector * 980.f; // Default to "regular" gravity
 								
 		// NOTE: We assume gravity is -Z for perf... If we want to support buoyancy for
@@ -1227,8 +1290,8 @@ namespace BuoyancyAlgorithms
 
 		// Scale submerged volume
 		float ScaledSubmergedVol = SubmergedVol;
-		float TotalVol;
-		ScaleSubmergedVolume(Evolution, RigidParticle, ScaledSubmergedVol, TotalVol);
+		float TotalParticleVol;
+		ScaleSubmergedVolume(Evolution, RigidParticle, ScaledSubmergedVol, TotalParticleVol);
 
 		// Compute buoyant force
 		//
@@ -1252,12 +1315,18 @@ namespace BuoyancyAlgorithms
 		const FVec3 WorldCoM = RigidGeneric->PCom();
 		const FVec3 CoMDiff = SubmergedCoM - WorldCoM;
 
-		// #todo(dmp): should we push water along normal or -gravity
-		OutWorldBuoyantForce = WaterN * BuoyantForceMagnitude;
-				
+		// optionally apply buoyancy along normal or -gravity
+		if (bBuoyancyAlgorithmsBuoyancyAlongNormal)
+		{
+			OutWorldBuoyantForce = WaterN * BuoyantForceMagnitude;
+		}
+		else
+		{
+			OutWorldBuoyantForce = -1.0 * GravityDir * BuoyantForceMagnitude;
+		}
+
 		// Compute world buoyant force and torque	
 		// buoyancy is computed at the submerged center of mass		
 		OutWorldBuoyantTorque = Chaos::FVec3::CrossProduct(CoMDiff, OutWorldBuoyantForce);				
 	}
-}
-
+}
\ No newline at end of file
diff --git a/Engine/Plugins/Experimental/Buoyancy/Source/Runtime/Private/BuoyancySubsystem.cpp b/Engine/Plugins/Experimental/Buoyancy/Source/Runtime/Private/BuoyancySubsystem.cpp
--- a/Engine/Plugins/Experimental/Buoyancy/Source/Runtime/Private/BuoyancySubsystem.cpp
+++ b/Engine/Plugins/Experimental/Buoyancy/Source/Runtime/Private/BuoyancySubsystem.cpp
@@ -61,6 +61,14 @@ static FAutoConsoleVariableRef CVarUseShallowWaterSimulation(
 	ECVF_Scalability
 );
 
+int32 bUseAccurateIntegrationForSplines = 0;
+static FAutoConsoleVariableRef CVarUseAccurateIntegrationForSplines(
+	TEXT("p.Buoyancy.bUseAccurateIntegrationForSplines"),
+	bUseAccurateIntegrationForSplines,
+	TEXT("new buoancy method for splines"),
+	ECVF_Scalability
+);
+
 //
 // Logging
 //
@@ -695,8 +703,7 @@ void FBuoyancySubsystemSimCallback::OnMidPhaseModification_Internal(Chaos::FMidP
 
 	// Apply buoyant forces resulting from submersions
 	ApplyBuoyantForces(*Evolution);
-
-	// #todo(dmp): need to generate callback info from shallow water interactions
+	
 	// Generate async outputs for callback data
 	GenerateCallbackData();
 
@@ -781,10 +788,8 @@ void FBuoyancySubsystemSimCallback::TrackInteraction(
 		MidPhase.Disable();
 	}
 	
-	// #todo(dmp): ideally for shallow water, we'd be sampling the body a bit more to tell if anything other than the center is interacting with water
-	float ClosestSplineKey = 0;
-	FVector ClosestPoint = FVector(0, 0, 0);
-	
+	TSharedPtr<FBuoyancyWaterSampler> WaterSampler = nullptr;
+
 	if (WaterSpline.ShouldSampleFromShallowWaterSimulation())
 	{
 		// eval shallow water - closest point is just the current point adjusted upward in Z
@@ -793,20 +798,31 @@ void FBuoyancySubsystemSimCallback::TrackInteraction(
 		float WaterHeight;
 		float WaterDepth;
 		WaterSpline.ShallowWaterSimData->SampleShallowWaterSimulationAtPosition(ParticlePos, WaterVelocity, WaterHeight, WaterDepth);
-		ClosestPoint = ParticlePos;
-		ClosestPoint.Z = WaterHeight;
 		
-		// if there is no water here, then return.  This is an approximation that there is water underneath this rigid that might cause interaction
+		// #todo(dmp): if there is no water here, then return.  This is an approximation that there is water underneath this rigid that might cause interaction
 		// ideally, we'd multisample this since we can lose some water interactions this way
 		if (WaterDepth < 1e-5)
 		{
 			return;
 		}
+
+		FVector ClosestWaterPointToParticle = ParticlePos;
+		ClosestWaterPointToParticle.Z = WaterHeight;
+
+		FShallowWaterSimulationGrid DefaultShallowWaterSimulationGrid;
+		const FShallowWaterSimulationGrid& ShallowWaterSimData = WaterSpline.ShallowWaterSimData.Get(DefaultShallowWaterSimulationGrid);
+
+		WaterSampler = TSharedPtr<FBuoyancyWaterSampler>(new FBuoyancyShallowWaterSampler(ShallowWaterSimData, ClosestWaterPointToParticle, WaterVelocity));
 	}
 	else
 	{
 		SCOPE_CYCLE_COUNTER(STAT_BuoyancySubsystem_SplineEvaluation)
 
+		float ClosestSplineKey = 0;
+		FVector ClosestPoint = FVector::ZeroVector;
+		FVector WaterN = FVector::ZeroVector;
+		FVector WaterVel = FVector::ZeroVector;
+
 		// Find water surface at the nearest point on the spline
 		const FVector ParticlePos = RigidParticle->XCom();
 		{
@@ -827,44 +843,140 @@ void FBuoyancySubsystemSimCallback::TrackInteraction(
 			SCOPE_CYCLE_COUNTER(STAT_BuoyancySubsystem_SplineEvaluation_Eval)
 			ClosestPoint = WaterSpline.Transform.TransformPosition(WaterSpline.Position.Eval(ClosestSplineKey));
 		}
+
+		// Get water surface level and normal
+		// Don't build interaction if outside of water body
+		{			
+			// Find water surface at the nearest point on the spline			
+			FVector ClosestPointDerivative = WaterSpline.Transform.TransformVector(
+			WaterSpline.Position.EvalDerivative(ClosestSplineKey));
+
+			// Water normal direction depends on body type
+			if (WaterSpline.BodyType == EWaterBodyType::River)
+			{
+				// River water normal can be determined by the relationship of the 
+				// derivative of the spline position to the up vector.
+				//
+				// NOTE: This calculation breaks down in the limit of purely
+				// vertical water
+				const FVector SplineRight = FVector::CrossProduct(FVector::UpVector, ClosestPointDerivative);
+				const FVector SplineUp = FVector::CrossProduct(ClosestPointDerivative, SplineRight);
+				WaterN = SplineUp.GetSafeNormal();
+			}
+			else
+			{
+				WaterN = FVector::UpVector;
+			}
+
+			// Project the position difference onto the water surface
+			const FVector Diff = ClosestPoint - ParticlePos;
+			const FVector LateralDiff = Diff - (WaterN * FVector::DotProduct(WaterN, Diff));
+
+			// Different water body types have different ways of determining
+			// whether a point is laterally inside their volume.
+			switch (WaterSpline.BodyType)
+			{
+				case EWaterBodyType::River:
+				{
+					if (WaterSpline.Width.IsSet())
+					{
+						// If distance to spline is greater than the width of the spline,
+						// then this is a river and we're outside of it.
+						const float Width = WaterSpline.Width->Eval(ClosestSplineKey);
+						const float DistSq = FVector::DotProduct(LateralDiff, LateralDiff);
+						const float WidthSq = Width * Width * .25f;
+						if (DistSq > WidthSq) { return; }
+					}
+					break;
+				}
+
+				case EWaterBodyType::Lake:
+				{
+					// Determine if we're inside the lake by projecting the horizontal spline
+					// diff onto the cross product of the spline direction and the up-vector
+					// (ie, the right-vector)
+					const FVector RightVector = FVector::CrossProduct(ClosestPointDerivative, FVector::UpVector);
+					const float DiffProj = FVector::DotProduct(RightVector, LateralDiff);
+					if (DiffProj < SMALL_NUMBER) { return; }
+					break;
+				}
+			}
+
+			WaterSampler = TSharedPtr<FBuoyancyWaterSampler>(new FBuoyancyConstantSplineSampler(WaterSpline, ClosestPoint, ClosestPointDerivative, ClosestSplineKey, WaterN));
+
+#if ENABLE_DRAW_DEBUG
+			if (bBuoyancyDebugDraw)
+			{
+				// Spline Color
+				const FColor SplineColor = FColor::Cyan;
+
+				// Draw projection onto the line
+				const FVector SurfacePoint = ClosestPoint - LateralDiff;
+				Chaos::FDebugDrawQueue::GetInstance().DrawDebugLine(ParticlePos, SurfacePoint, SplineColor, false, -1.f, -1, 6.f);
+				Chaos::FDebugDrawQueue::GetInstance().DrawDebugLine(SurfacePoint, ClosestPoint, SplineColor, false, -1.f, -1, 3.f);
+
+				// Draw a section of the spline near the spline key
+				Chaos::FVec3 PrevPoint;
+				bool bFirst = true;
+				for (float SplineKey = ClosestSplineKey - .1f; SplineKey <= ClosestSplineKey + .1f; SplineKey += .05f)
+				{
+					const FVector SplinePoint = WaterSpline.Transform.TransformPosition(WaterSpline.Position.Eval(SplineKey));
+					if (bFirst)
+					{
+						bFirst = false;
+					}
+					else
+					{
+						Chaos::FDebugDrawQueue::GetInstance().DrawDebugDirectionalArrow(PrevPoint, SplinePoint, 15.f, SplineColor, false, -1.f, -1, 3.f);
+					}
+					PrevPoint = SplinePoint;
+				}
+
+				// Draw water surface normal at the surface point
+				Chaos::FDebugDrawQueue::GetInstance().DrawDebugDirectionalArrow(SurfacePoint, SurfacePoint + (WaterN * 100.f), 20.f, FColor::Green, false, -1.f, -1, 3.f);
+			}
+#endif
+		}
+	
+		
 	}
 
 	{
 		SCOPE_CYCLE_COUNTER(STAT_BuoyancySubsystem_AddInteraction)
 
 		// If this particle already has a list of interactions, add to it. Otherwise
-		// make a new one.
-		// #todo(dmp): we should split out the interaction data here to be for shallow water or not
+		// make a new one.		
 		BuoyancyParticleData.GetInteractions(*RigidParticle).Add({
 			RigidParticle,
 			WaterParticle,
-			WaterSpline,
-			ClosestSplineKey,
-			ClosestPoint,
+			WaterSampler
 		});
 	}
 }
 
 void FBuoyancySubsystemSimCallback::ProcessInteractions(Chaos::FPBDRigidsEvolution& Evolution)
 {
+	SCOPE_CYCLE_COUNTER(STAT_BuoyancySubsystem_ProcessInteractions)
+
 	for (FBuoyancyInteractionArray& ParticleInteractions : BuoyancyParticleData.Interactions)
 	{
 		// Sort this particle's water interactions by water level - highest to lowest
 		ParticleInteractions.Sort([](const FBuoyancyInteraction& A, const FBuoyancyInteraction& B)
 		{
-			return A.ClosestPoint.Z > B.ClosestPoint.Z;
+			return A.WaterSampler->GetClosestWaterPointToParticle().Z > B.WaterSampler->GetClosestWaterPointToParticle().Z;
 		});
 
 		// Process each interaction
 		for (FBuoyancyInteraction& Interaction : ParticleInteractions)
-		{						
-			// #todo(dmp): we do not want to allow multiple shallow water sims interacting with a single particle for now.
-			// #todo(dmp): correctly handle when a particle can overlap shallow water and other water bodies at the same time
-			// in order to do this, we must unite the integration methods for shallow water and spline based water bodies
-			// and rework the submersion data structures.
-			if (Interaction.WaterSpline.ShouldSampleFromShallowWaterSimulation())
+		{
+			if (!Interaction.WaterSampler.IsValid())
+			{
+				UE_LOG(LogBuoyancySubsystem, Warning, TEXT("Skipped invalid buoyancy interaction"));
+				return;
+			}
+			else if (Interaction.WaterSampler->GetSamplerType() == EWaterSamplerType::ShallowWater || bUseAccurateIntegrationForSplines)
 			{
-				ProcessInteractionShallowWater(Evolution, Interaction);
+				ProcessAccurateInteraction(Evolution, Interaction, Interaction.WaterSampler);
 			}
 			else
 			{
@@ -874,29 +986,21 @@ void FBuoyancySubsystemSimCallback::ProcessInteractions(Chaos::FPBDRigidsEvoluti
 	}
 }
 
-void FBuoyancySubsystemSimCallback::ProcessInteractionShallowWater(Chaos::FPBDRigidsEvolution& Evolution, FBuoyancyInteraction& Interaction)
+void FBuoyancySubsystemSimCallback::ProcessAccurateInteraction(Chaos::FPBDRigidsEvolution& Evolution, FBuoyancyInteraction& Interaction, TSharedPtr<FBuoyancyWaterSampler> WaterSampler)
 {		
-	const Chaos::FReal DeltaSeconds = GetDeltaTime_Internal();
-		
-	if (!Interaction.WaterSpline.ShallowWaterSimData.IsSet())
-	{
-		return;
-	}
+	const Chaos::FReal DeltaSeconds = GetDeltaTime_Internal();	
 
-	FShallowWaterSimulationGrid DefaultShallowWaterSimulationGrid;
-	const FShallowWaterSimulationGrid& ShallowWaterSimData = Interaction.WaterSpline.ShallowWaterSimData.Get(DefaultShallowWaterSimulationGrid);
-
-	float TotalSubmergedVol;
-	Chaos::FVec3 TotalSubmergedCoM;
-	Chaos::FVec3 TotalWorldForce;
-	Chaos::FVec3 TotalWorldTorque;
-		
-	// #todo(dmp): abstract out the ShallowWaterSimData to generic class so we can sample other velocity sources (ie: splines, constant value, etc)
+	float SubmergedVol;
+	Chaos::FVec3 SubmergedCoM;
+	Chaos::FVec3 WorldForce;
+	Chaos::FVec3 WorldTorque;		
+	float TotalParticleVol;
+	
 	BuoyancyAlgorithms::ComputeSubmergedVolumeAndForcesForParticle(BuoyancyParticleData,
 		Interaction.RigidParticle, Interaction.WaterParticle,
-		ShallowWaterSimData,
+		WaterSampler,
 		Evolution, DeltaSeconds, BuoyancySettings->WaterDensity, BuoyancySettings->WaterDrag,
-		TotalSubmergedVol, TotalSubmergedCoM, TotalWorldForce, TotalWorldTorque);
+		TotalParticleVol, SubmergedVol, SubmergedCoM, WorldForce, WorldTorque);
 
 	// Apply forces to the particle
 	{
@@ -907,10 +1011,10 @@ void FBuoyancySubsystemSimCallback::ProcessInteractionShallowWater(Chaos::FPBDRi
 #if ENABLE_DRAW_DEBUG
 		if (bBuoyancyDebugDraw)
 		{
-			Chaos::FDebugDrawQueue::GetInstance().DrawDebugSphere(TotalSubmergedCoM, 20, 10, FColor::Green, false, -1.f, -1, 2.0f);
+			Chaos::FDebugDrawQueue::GetInstance().DrawDebugSphere(SubmergedCoM, 20, 10, FColor::Green, false, -1.f, -1, 2.0f);
 			Chaos::FDebugDrawQueue::GetInstance().DrawDebugSphere(RigidGeneric->PCom(), 20, 10, FColor::Blue, false, -1.f, -1, 2.0f);
 
-			Chaos::FDebugDrawQueue::GetInstance().DrawDebugDirectionalArrow(RigidGeneric->PCom(), RigidGeneric->PCom() + TotalWorldForce, 20, FColor::Blue, false, -1.f, -1, 2.f);
+			Chaos::FDebugDrawQueue::GetInstance().DrawDebugDirectionalArrow(RigidGeneric->PCom(), RigidGeneric->PCom() + WorldForce, 20, FColor::Blue, false, -1.f, -1, 2.f);
 		}
 #endif
 
@@ -918,139 +1022,81 @@ void FBuoyancySubsystemSimCallback::ProcessInteractionShallowWater(Chaos::FPBDRi
 		// This will change to integrate the accumulated accelerations instead of forces
 
 		// Use inertia to convert forces to accelerations
-		const Chaos::FVec3 LinearAccel = RigidGeneric->InvM() * TotalWorldForce;
+		const Chaos::FVec3 LinearAccel = RigidGeneric->InvM() * WorldForce;
 
 		const Chaos::FMatrix33 WorldInvI = Chaos::Utilities::ComputeWorldSpaceInertia(RigidGeneric->RCom(), RigidGeneric->ConditionedInvI());
-		const Chaos::FVec3 AngularAccel = WorldInvI * TotalWorldTorque;
+		const Chaos::FVec3 AngularAccel = WorldInvI * WorldTorque;
 
 		// Integrate to get delta velocities
 		Chaos::FVec3 DeltaV = LinearAccel * DeltaSeconds;
 		Chaos::FVec3 DeltaW = AngularAccel * DeltaSeconds;
-	
-		// Clamp delta velocities
-		DeltaV = DeltaV.GetClampedToSize(0.f, BuoyancySettings->MaxDeltaV);
-		DeltaW = DeltaW.GetClampedToSize(0.f, BuoyancySettings->MaxDeltaW);
-
-		// Compute new velocity
-		const Chaos::FVec3 NewV = Interaction.RigidParticle->GetV() + DeltaV;
-		const Chaos::FVec3 NewW = Interaction.RigidParticle->GetW() + DeltaW;
-
-		// Apply the deltas
-		Interaction.RigidParticle->SetV(NewV);
-		Interaction.RigidParticle->SetW(NewW);
-	
-		// Wake up the body??
-		if (BuoyancySettings->bKeepAwake)
+		
+		if (SubmergedVol > SMALL_NUMBER)
 		{
-			Evolution.SetParticleObjectState(Interaction.RigidParticle, Chaos::EObjectStateType::Dynamic);
-		}
-	}
-}
-
-void FBuoyancySubsystemSimCallback::ProcessInteraction(Chaos::FPBDRigidsEvolution& Evolution, FBuoyancyInteraction& Interaction)
-{
+			// Get the submersion for this particle. If it's new, the particle ptr will be null.
+			// If it's not new, make sure the particle ptr is the right particle.
+			FBuoyancySubmersion& Submersion = BuoyancyParticleData.GetSubmersion(*Interaction.RigidParticle);
+			if (Submersion.Particle != nullptr)
+			{
+				ensureMsgf(Submersion.Particle == Interaction.RigidParticle, TEXT("Something went wrong - there's a particle index mismatch in the Submersions sparse array"));
+			}
 
-	// Get water surface level and normal
-	FVector WaterN;
-	FVector ClosestPosDerivative;
-	{
-		SCOPE_CYCLE_COUNTER(STAT_BuoyancySubsystem_SplineEvaluation)
+			Submersion.Particle = Interaction.RigidParticle;
 
-		// Find water surface at the nearest point on the spline
-		const FVector ParticlePos = Interaction.RigidParticle->XCom();
-		ClosestPosDerivative = Interaction.WaterSpline.Transform.TransformVector(
-			Interaction.WaterSpline.Position.EvalDerivative(Interaction.ClosestSplineKey));
+			Chaos::FSingleParticlePhysicsProxy* RigidProxy = static_cast<Chaos::FSingleParticlePhysicsProxy*>(Submersion.Particle->PhysicsProxy());
+			Submersion.SyncTimestamp = RigidProxy->GetSyncTimestamp();
 
-		// Water normal direction depends on body type
-		if (Interaction.WaterSpline.BodyType == EWaterBodyType::River)
-		{
-			// River water normal can be determined by the relationship of the 
-			// derivative of the spline position to the up vector.
-			//
-			// NOTE: This calculation breaks down in the limit of purely
-			// vertical water
-			const FVector SplineRight = FVector::CrossProduct(FVector::UpVector, ClosestPosDerivative);
-			const FVector SplineUp = FVector::CrossProduct(ClosestPosDerivative, SplineRight);
-			WaterN = SplineUp.GetSafeNormal();
-		}
-		else
-		{
-			WaterN = FVector::UpVector;
-		}
+			// Get the weighted-average CoM
+			// NOTE: The unchecked division should be safe since we already
+			// know SubmergedVol > SMALL_NUMBER
+			const float TotalSubmergedVol = Submersion.Vol + SubmergedVol;
+			Submersion.CoM = ((Submersion.CoM * Submersion.Vol) + (SubmergedCoM * SubmergedVol)) / TotalSubmergedVol;
 
-		// Project the position difference onto the water surface
-		const FVector Diff = Interaction.ClosestPoint - ParticlePos;
-		const FVector LateralDiff = Diff - (WaterN * FVector::DotProduct(WaterN, Diff));
+			// Sum the volumes
+			Submersion.Vol = TotalSubmergedVol;
 
-		// Different water body types have different ways of determining
-		// whether a point is laterally inside their volume.
-		switch (Interaction.WaterSpline.BodyType)
-		{
-			case EWaterBodyType::River:
-			{
-				if (Interaction.WaterSpline.Width.IsSet())
-				{
-					// If distance to spline is greater than the width of the spline,
-					// then this is a river and we're outside of it.
-					const float Width = Interaction.WaterSpline.Width->Eval(Interaction.ClosestSplineKey);
-					const float DistSq = FVector::DotProduct(LateralDiff, LateralDiff);
-					const float WidthSq = Width * Width * .25f;
-					if (DistSq > WidthSq) { return; }
-				}
-				break;
-			}
+			// Get the weighted-average slerped water surface norm
+			const float VolRatio = SubmergedVol / TotalSubmergedVol;
+						
+			// #todo(dmp): Split out submersion since we don't need norm here- only deltav and deltaw
+			// #todo(dmp): output the perimeter of intersection
+			Submersion.Vel = FMath::Lerp(Submersion.Vel, WaterSampler->GetAverageWaterVelocity(), VolRatio);
+			
+			Submersion.HasIntegratedQuantities = true;
+			Submersion.DeltaV = FMath::Lerp(Submersion.DeltaV, DeltaV, VolRatio);
+			Submersion.DeltaW = FMath::Lerp(Submersion.DeltaW, DeltaW, VolRatio);
 
-			case EWaterBodyType::Lake:
-			{
-				// Determine if we're inside the lake by projecting the horizontal spline
-				// diff onto the cross product of the spline direction and the up-vector
-				// (ie, the right-vector)
-				const FVector RightVector = FVector::CrossProduct(ClosestPosDerivative, FVector::UpVector);
-				const float DiffProj = FVector::DotProduct(RightVector, LateralDiff);
-				if (DiffProj < SMALL_NUMBER) { return; }
-				break;
-			}
+			// potentially add a surface callback
+			AddSurfaceTouchCallback(Interaction, Submersion, TotalParticleVol, SubmergedVol, SubmergedCoM);
 		}
+	}
+}
 
-#if ENABLE_DRAW_DEBUG
-		if (bBuoyancyDebugDraw)
-		{
-			// Spline Color
-			const FColor SplineColor = FColor::Cyan;
-
-			// Draw projection onto the line
-			const FVector SurfacePoint = Interaction.ClosestPoint - LateralDiff;
-			Chaos::FDebugDrawQueue::GetInstance().DrawDebugLine(ParticlePos, SurfacePoint, SplineColor, false, -1.f, -1, 6.f);
-			Chaos::FDebugDrawQueue::GetInstance().DrawDebugLine(SurfacePoint, Interaction.ClosestPoint, SplineColor, false, -1.f, -1, 3.f);
-
-			// Draw a section of the spline near the spline key
-			Chaos::FVec3 PrevPoint;
-			bool bFirst = true;
-			for (float SplineKey = Interaction.ClosestSplineKey - .1f; SplineKey <= Interaction.ClosestSplineKey + .1f; SplineKey += .05f)
-			{
-				const FVector SplinePoint = Interaction.WaterSpline.Transform.TransformPosition(Interaction.WaterSpline.Position.Eval(SplineKey));
-				if (bFirst)
-				{
-					bFirst = false;
-				}
-				else
-				{
-					Chaos::FDebugDrawQueue::GetInstance().DrawDebugDirectionalArrow(PrevPoint, SplinePoint, 15.f, SplineColor, false, -1.f, -1, 3.f);
-				}
-				PrevPoint = SplinePoint;
-			}
+void FBuoyancySubsystemSimCallback::ProcessInteraction(Chaos::FPBDRigidsEvolution& Evolution, FBuoyancyInteraction& Interaction)
+{
+	SCOPE_CYCLE_COUNTER(STAT_BuoyancySubsystem_ProcessInteraction)
 
-			// Draw water surface normal at the surface point
-			Chaos::FDebugDrawQueue::GetInstance().DrawDebugDirectionalArrow(SurfacePoint, SurfacePoint + (WaterN * 100.f), 20.f, FColor::Green, false, -1.f, -1, 3.f);
-		}
-#endif
+	TSharedPtr<FBuoyancyConstantSplineSampler> ConstantSplineSampler;
+	if (!Interaction.WaterSampler.IsValid())
+	{
+		UE_LOG(LogBuoyancySubsystem, Warning, TEXT("Skipped invalid buoyancy interaction "));
+		return;
 	}
-
+	else if (Interaction.WaterSampler->GetSamplerType() == EWaterSamplerType::ConstantSpline)
+	{
+		ConstantSplineSampler = StaticCastSharedPtr<FBuoyancyConstantSplineSampler>(Interaction.WaterSampler);		
+	}
+	else
+	{
+		UE_LOG(LogBuoyancySubsystem, Warning, TEXT("Skipped buoyancy interaction that isn't based on water splines"));
+		return;
+	}
+	
 	// Compute submerged volume and CoM
 	float SubmergedVol;
 	Chaos::FVec3 SubmergedCoM;
 	float TotalVol;
-	if (BuoyancyAlgorithms::ComputeSubmergedVolume(BuoyancyParticleData, Evolution, Interaction.RigidParticle, Interaction.WaterParticle, Interaction.ClosestPoint, WaterN, BuoyancySettings->MaxNumBoundsSubdivisions, BuoyancySettings->MinBoundsSubdivisionVol, SubmergedVol, SubmergedCoM, TotalVol))
+	if (BuoyancyAlgorithms::ComputeSubmergedVolume(BuoyancyParticleData, Evolution, Interaction.RigidParticle, Interaction.WaterParticle, ConstantSplineSampler->GetClosestWaterPointToParticle(), ConstantSplineSampler->WaterN, BuoyancySettings->MaxNumBoundsSubdivisions, BuoyancySettings->MinBoundsSubdivisionVol, SubmergedVol, SubmergedCoM, TotalVol))
 	{
 		SCOPE_CYCLE_COUNTER(STAT_BuoyancySubsystem_BuildSubmersions)
 
@@ -1058,11 +1104,8 @@ void FBuoyancySubsystemSimCallback::ProcessInteraction(Chaos::FPBDRigidsEvolutio
 		// approximate submerged center of mass, and apply the buoyancy force there.
 		if (SubmergedVol > SMALL_NUMBER)
 		{
-			// Get the water velocity at this point
-			const FVector WaterVel
-				= Interaction.WaterSpline.Velocity.IsSet()
-				? Interaction.WaterSpline.Velocity->Eval(Interaction.ClosestSplineKey) * ClosestPosDerivative.GetSafeNormal()
-				: Chaos::FVec3::ZeroVector;
+			// sample water spline for velocity with this call
+			ConstantSplineSampler->InitializeForSubmersion();
 
 			// Get the submersion for this particle. If it's new, the particle ptr will be null.
 			// If it's not new, make sure the particle ptr is the right particle.
@@ -1088,37 +1131,43 @@ void FBuoyancySubsystemSimCallback::ProcessInteraction(Chaos::FPBDRigidsEvolutio
 
 			// Get the weighted-average slerped water surface norm
 			const float VolRatio = SubmergedVol / TotalSubmergedVol;
-			const FQuat NormRot = FQuat::FindBetweenNormals(Submersion.Norm, WaterN);
+			const FQuat NormRot = FQuat::FindBetweenNormals(Submersion.Norm, ConstantSplineSampler->WaterN);
 			Submersion.Norm = NormRot.Slerp(FQuat::Identity, NormRot, VolRatio).GetUpVector();
 
 			// Blend water velocities
-			Submersion.Vel = FMath::Lerp(Submersion.Vel, WaterVel, VolRatio);
+			Submersion.Vel = FMath::Lerp(Submersion.Vel, ConstantSplineSampler->WaterVel, VolRatio);
 
-			// If this is a surface touch record it for callback, if 
-			if (BuoyancySettings->SurfaceTouchCallbackFlags != 0 &&
-				TotalVol > SubmergedVol * (1.f + UE_KINDA_SMALL_NUMBER))
-			{
-				SCOPE_CYCLE_COUNTER(STAT_BuoyancySubsystem_BuildSubmersionCallbackData)
+			// potentially add a surface callback
+			AddSurfaceTouchCallback(Interaction, Submersion, TotalVol, SubmergedVol, SubmergedCoM);
+		}
+	}
+}
 
-				// Proceed only if this CoM is moving fast enough to generate events
-				const Chaos::FVec3 CoMDiff = Submersion.CoM - Interaction.RigidParticle->XCom();
-				const Chaos::FVec3 CoMVel = Interaction.RigidParticle->GetV() + Chaos::FVec3::CrossProduct(Interaction.RigidParticle->GetW(), CoMDiff);
-				const float CoMVelSq = Chaos::FVec3::DotProduct(CoMVel, CoMVel);
-				const float MinVel = BuoyancySettings->MinVelocityForSurfaceTouchCallback;
-				const float MinVelSq = MinVel * MinVel;
+void FBuoyancySubsystemSimCallback::AddSurfaceTouchCallback(FBuoyancyInteraction& Interaction, FBuoyancySubmersion& Submersion, float TotalVol, float SubmergedVol, FVector SubmergedCoM)
+{
+	// If this is a surface touch record it for callback, if 
+	if (BuoyancySettings->SurfaceTouchCallbackFlags != 0 &&
+		TotalVol > SubmergedVol * (1.f + UE_KINDA_SMALL_NUMBER))
+	{
+		SCOPE_CYCLE_COUNTER(STAT_BuoyancySubsystem_BuildSubmersionCallbackData)
 
-				if (CoMVelSq > MinVelSq)
-				{
-					// If we don't have a metadata for this particle yet, add one
-					FBuoyancySubmersionMetaData& MetaData = BuoyancyParticleData.GetSubmersionMetaData(*Interaction.RigidParticle);
+			// Proceed only if this CoM is moving fast enough to generate events
+			const Chaos::FVec3 CoMDiff = Submersion.CoM - Interaction.RigidParticle->XCom();
+		const Chaos::FVec3 CoMVel = Interaction.RigidParticle->GetV() + Chaos::FVec3::CrossProduct(Interaction.RigidParticle->GetW(), CoMDiff);
+		const float CoMVelSq = Chaos::FVec3::DotProduct(CoMVel, CoMVel);
+		const float MinVel = BuoyancySettings->MinVelocityForSurfaceTouchCallback;
+		const float MinVelSq = MinVel * MinVel;
 
-					// If we haven't already maxed out on water contacts, add one
-					if (MetaData.WaterContacts.Num() < FBuoyancySubmersionMetaData::MaxNumWaterContacts)
-					{
-						Chaos::FSingleParticlePhysicsProxy* WaterProxy = static_cast<Chaos::FSingleParticlePhysicsProxy*>(Interaction.WaterParticle->PhysicsProxy());
-						MetaData.WaterContacts.Add({ Interaction.WaterParticle, WaterProxy->GetSyncTimestamp(), SubmergedVol, SubmergedCoM, CoMVel});
-					}
-				}
+		if (CoMVelSq > MinVelSq)
+		{
+			// If we don't have a metadata for this particle yet, add one
+			FBuoyancySubmersionMetaData& MetaData = BuoyancyParticleData.GetSubmersionMetaData(*Interaction.RigidParticle);
+
+			// If we haven't already maxed out on water contacts, add one
+			if (MetaData.WaterContacts.Num() < FBuoyancySubmersionMetaData::MaxNumWaterContacts)
+			{
+				Chaos::FSingleParticlePhysicsProxy* WaterProxy = static_cast<Chaos::FSingleParticlePhysicsProxy*>(Interaction.WaterParticle->PhysicsProxy());
+				MetaData.WaterContacts.Add({ Interaction.WaterParticle, WaterProxy->GetSyncTimestamp(), SubmergedVol, SubmergedCoM, CoMVel });
 			}
 		}
 	}
@@ -1137,25 +1186,44 @@ void FBuoyancySubsystemSimCallback::ApplyBuoyantForces(Chaos::FPBDRigidsEvolutio
 	// Apply all buoyant forces
 	for (const FBuoyancySubmersion& Submersion : BuoyancyParticleData.Submersions)
 	{
-		// Figure out the gravity level of the particle
-		const int32 GravityGroupIndex = Submersion.Particle->GravityGroupIndex();
-		const Chaos::FVec3 GravityAccel
-			= PerParticleGravity != nullptr && GravityGroupIndex != INDEX_NONE
-			? (Chaos::FVec3)PerParticleGravity->GetAcceleration(GravityGroupIndex)
-			: Chaos::FVec3::DownVector * 980.f; // Default to "regular" gravity
-
-		// Compute delta linear and angular velocities due to buoyancy. If they're big enough to
-		// matter, apply them
-		Chaos::FVec3 DeltaV, DeltaW;
-		if (BuoyancyAlgorithms::ComputeBuoyantForce(Submersion.Particle, DeltaSeconds, BuoyancySettings->WaterDensity, BuoyancySettings->WaterDrag, GravityAccel, Submersion.CoM, Submersion.Vol, Submersion.Vel, Submersion.Norm, DeltaV, DeltaW))
+		Chaos::FVec3 DeltaV;
+		Chaos::FVec3 DeltaW;
+
+		bool IsValidForce = true;
+
+		// If the submersion has integrated velocity and torque, use that, otherwise, compute it
+		if (Submersion.HasIntegratedQuantities)
+		{			
+			DeltaV = Submersion.DeltaV;
+			DeltaW = Submersion.DeltaW;							
+		}
+		else
+		{
+			// Figure out the gravity level of the particle
+			const int32 GravityGroupIndex = Submersion.Particle->GravityGroupIndex();
+			const Chaos::FVec3 GravityAccel
+				= PerParticleGravity != nullptr && GravityGroupIndex != INDEX_NONE
+				? (Chaos::FVec3)PerParticleGravity->GetAcceleration(GravityGroupIndex)
+				: Chaos::FVec3::DownVector * 980.f; // Default to "regular" gravity
+
+			// Compute delta linear and angular velocities due to buoyancy. If they're big enough to
+			// matter, apply them
+			IsValidForce = BuoyancyAlgorithms::ComputeBuoyantForce(Submersion.Particle, DeltaSeconds,
+				BuoyancySettings->WaterDensity, BuoyancySettings->WaterDrag, GravityAccel, Submersion.CoM, Submersion.Vol, Submersion.Vel, Submersion.Norm, DeltaV, DeltaW);
+		}
+
+		if (IsValidForce)
 		{
-			// Clamp delta velocities
 			DeltaV = DeltaV.GetClampedToSize(0.f, BuoyancySettings->MaxDeltaV);
 			DeltaW = DeltaW.GetClampedToSize(0.f, BuoyancySettings->MaxDeltaW);
 
+			// Compute new velocity
+			const Chaos::FVec3 NewV = Submersion.Particle->GetV() + DeltaV;
+			const Chaos::FVec3 NewW = Submersion.Particle->GetW() + DeltaW;
+
 			// Apply the deltas
-			Submersion.Particle->SetV(Submersion.Particle->GetV() + DeltaV);
-			Submersion.Particle->SetW(Submersion.Particle->GetW() + DeltaW);
+			Submersion.Particle->SetV(NewV);
+			Submersion.Particle->SetW(NewW);
 
 			// Wake up the body??
 			if (BuoyancySettings->bKeepAwake)
diff --git a/Engine/Plugins/Experimental/Buoyancy/Source/Runtime/Public/BuoyancyAlgorithms.h b/Engine/Plugins/Experimental/Buoyancy/Source/Runtime/Public/BuoyancyAlgorithms.h
--- a/Engine/Plugins/Experimental/Buoyancy/Source/Runtime/Public/BuoyancyAlgorithms.h
+++ b/Engine/Plugins/Experimental/Buoyancy/Source/Runtime/Public/BuoyancyAlgorithms.h
@@ -13,6 +13,7 @@ namespace Chaos
 	class FPBDRigidsEvolutionGBF;
 }
 
+class FBuoyancyWaterSampler;
 struct FBuoyancyParticleData;
 
 namespace BuoyancyAlgorithms
@@ -71,9 +72,9 @@ namespace BuoyancyAlgorithms
 	// given a particle, loop over the contained shapes and accumulate force/torque/submerged CoM values
 	void ComputeSubmergedVolumeAndForcesForParticle(FBuoyancyParticleData& ParticleData,
 		const Chaos::FGeometryParticleHandle* SubmergedParticle, const Chaos::FGeometryParticleHandle* WaterParticle,
-		const FShallowWaterSimulationGrid& ShallowWaterGrid,
+		TSharedPtr<FBuoyancyWaterSampler> WaterSampler,
 		const Chaos::FPBDRigidsEvolution& Evolution, const float DeltaSeconds, const float WaterDensity, const float WaterDrag,
-		float& OutTotalSubmergedVol, Chaos::FVec3& OutTotalSubmergedCoM, Chaos::FVec3& OutTotalForce, Chaos::FVec3& OutTotalTorque);
+		float& OutTotalParticleVol, float& OutTotalSubmergedVol, Chaos::FVec3& OutTotalSubmergedCoM, Chaos::FVec3& OutTotalForce, Chaos::FVec3& OutTotalTorque);
 
 	// given a shape, compute the submerged volume and accumulate forces
 	// this is done in a single function call because of the iterative nature of the algorithm
@@ -81,18 +82,20 @@ namespace BuoyancyAlgorithms
 		const Chaos::FGeometryParticleHandle* SubmergedParticle,
 		const Chaos::FPBDRigidsEvolution& Evolution, float DeltaSeconds, const float WaterDensity, const float WaterDrag,
 		const Chaos::FImplicitObject* Implicit, const Chaos::FRigidTransform3& RelativeTransform, const int32 RootObjectIndex, const int32 ObjectIndex, const int32 LeafObjectIndex,
-		const FShallowWaterSimulationGrid& ShallowWaterGrid, 
+		TSharedPtr<FBuoyancyWaterSampler> WaterSampler,
 		Chaos::FVec3& OutWaterP, Chaos::FVec3& OutWaterN, 
 		float& OutSubmergedVol, Chaos::FVec3& OutSubmergedCoM,
 		Chaos::FVec3& OutForce, Chaos::FVec3& OutTorque);
 
 	// find intersection points between a plane and aabbox
-	void FindAllIntersectionPoints(const Chaos::FVec3& WaterP, const Chaos::FVec3& WaterN, const Chaos::FAABB3& LocalBox, const FVector WorldVertexPosition[8],
-		TMap<int, FVector> &OutEdgeToIntersectionPointMap, TArray<FVector>& OutOrderedIntersectionPoints, FVector& OutIntersectionCenter);
+	void FORCEINLINE FindAllIntersectionPoints(const Chaos::FVec3& WaterP, const Chaos::FVec3& WaterN, const Chaos::FAABB3& LocalBox, const FVector WorldVertexPosition[8],
+		FVector EdgeToIntersectionPoint[12], int32& NumIntersections, TArray<FVector, TFixedAllocator<6>>& OutOrderedIntersectionPoints, FVector& OutIntersectionCenter);
 
 	// sort intersection points by angle
-	void SortIntersectionPointsByAngle(const Chaos::FVec3& WaterP, const Chaos::FVec3& WaterN, const Chaos::FVec3 &IntersectionCenter, 
-		const TMap<int, FVector>& EdgeToIntersectionPointMap, TArray<FVector>& OutOrderedIntersectionPoints);
+	void FORCEINLINE SortIntersectionPointsByAngle(const Chaos::FVec3& WaterP, const Chaos::FVec3& WaterN, const Chaos::FVec3& IntersectionCenter,
+		const FVector EdgeToIntersectionPoint[12], const int32 NumIntersections, TArray<FVector, TFixedAllocator<6>>& OutOrderedIntersectionPoints);
+
+	bool FORCEINLINE EdgePlaneIntersection(const Chaos::FVec3& WaterP, const Chaos::FVec3& WaterN, const Chaos::FVec3& V0, const Chaos::FVec3& V1, Chaos::FVec3& IntersectionPoint);
 
 	// compute area and volume of a tet from a triangle and center point on mesh
 	void FORCEINLINE ComputeTriangleAreaAndVolume(const FVector &V0, const FVector &V1, const FVector &V2,
@@ -107,6 +110,6 @@ namespace BuoyancyAlgorithms
 		FVector& OutTotalWorldForce, FVector& OutTotalWorldTorque);
 
 	// compute the force the buoyancy exerts on a shape
-	void ComputeBuoyantForceForShape(const Chaos::FPBDRigidsEvolution& Evolution, const Chaos::FPBDRigidParticleHandle* RigidParticle, const float DeltaSeconds, const float WaterDensity,
+	void FORCEINLINE ComputeBuoyantForceForShape(const Chaos::FPBDRigidsEvolution& Evolution, const Chaos::FPBDRigidParticleHandle* RigidParticle, const float DeltaSeconds, const float WaterDensity,
 		const Chaos::FVec3& SubmergedCoM, const float SubmergedVol, const Chaos::FVec3& WaterN, Chaos::FVec3& OutWorldBuoyantForce, Chaos::FVec3& OutWorldBuoyantTorque);
 }
diff --git a/Engine/Plugins/Experimental/Buoyancy/Source/Runtime/Public/BuoyancyParticleData.h b/Engine/Plugins/Experimental/Buoyancy/Source/Runtime/Public/BuoyancyParticleData.h
--- a/Engine/Plugins/Experimental/Buoyancy/Source/Runtime/Public/BuoyancyParticleData.h
+++ b/Engine/Plugins/Experimental/Buoyancy/Source/Runtime/Public/BuoyancyParticleData.h
@@ -13,6 +13,132 @@
 #include "BuoyancyWaterSplineData.h"
 #include "BuoyancyStats.h"
 
+enum EWaterSamplerType
+{
+	ConstantSpline,
+	ShallowWater,
+	Invalid,
+};
+
+// abstract class to sample water for a given particle
+class FBuoyancyWaterSampler
+{
+public:
+	FBuoyancyWaterSampler() : SamplerType(EWaterSamplerType::Invalid), bUsePerShapeWaterPlane(true), 
+		ClosestWaterPointToParticle(0,0,0), AverageVelocity(0,0,0)
+	{
+	};
+
+	virtual ~FBuoyancyWaterSampler() {}
+
+	bool UsePerShapeWaterPlane() const { return bUsePerShapeWaterPlane;  }
+	FVector GetClosestWaterPointToParticle() const { return ClosestWaterPointToParticle; }
+
+	// average fluid velocity for the particle
+	FVector GetAverageWaterVelocity() const { return AverageVelocity; }
+	virtual void InitializeForSubmersion() {}
+	virtual void SampleWaterAtPosition(const FVector& WorldPosition, FVector& OutVelocity, float& OutHeight, float& OutDepth) const = 0;
+	virtual FVector SampleVelocityAtPosition(const FVector& WorldPosition) const = 0;
+	virtual FVector SampleNormalAtPosition(const FVector& WorldPosition) const = 0;
+
+	EWaterSamplerType GetSamplerType() const { return SamplerType; }
+protected:
+	EWaterSamplerType SamplerType;
+	bool bUsePerShapeWaterPlane;
+	FVector ClosestWaterPointToParticle;
+	FVector AverageVelocity;
+};
+
+// sample water from a baked shallow water simulation
+class FBuoyancyShallowWaterSampler : public FBuoyancyWaterSampler
+{
+public:
+	FBuoyancyShallowWaterSampler(const FShallowWaterSimulationGrid& InWaterSampler, const FVector& InClosestWaterPointToParticle, const FVector &InAverageVelocity) :
+		WaterSampler(InWaterSampler)
+	{
+		SamplerType = EWaterSamplerType::ShallowWater;
+		bUsePerShapeWaterPlane = true;
+		ClosestWaterPointToParticle = InClosestWaterPointToParticle;
+		AverageVelocity = InAverageVelocity;
+	}	
+
+	virtual void SampleWaterAtPosition(const FVector& WorldPosition, FVector& OutVelocity, float& OutHeight, float& OutDepth) const override
+	{
+		WaterSampler.SampleShallowWaterSimulationAtPosition(WorldPosition, OutVelocity, OutHeight, OutDepth);
+	}
+
+	virtual FVector SampleVelocityAtPosition(const FVector& WorldPosition) const override
+	{
+		FVector WaterVelocity;
+		float TmpWaterVal;
+		WaterSampler.SampleShallowWaterSimulationAtPosition(WorldPosition, WaterVelocity, TmpWaterVal, TmpWaterVal);
+
+		return WaterVelocity;
+	}
+
+	virtual FVector SampleNormalAtPosition(const FVector& WorldPosition) const override
+	{
+		return WaterSampler.ComputeShallowWaterSimulationNormalAtPosition(WorldPosition);
+	}
+
+private:
+	const FShallowWaterSimulationGrid& WaterSampler;
+};
+
+// sampler for spline rivers, but we are only doing 1 spline sample per interaction to 
+// match the perf of the old system.  Also, spline velocities don't vary much, so the net
+// effect isn't too different
+class FBuoyancyConstantSplineSampler : public FBuoyancyWaterSampler
+{
+public:
+	FBuoyancyConstantSplineSampler(const FBuoyancyWaterSplineData& InWaterSpline, const FVector& InClosestPoint, 
+		const FVector& InClosestPointDerivative, const float InClosestSplineKey, const FVector& InWaterN)
+		 : WaterSpline(InWaterSpline), ClosestPointDerivative(InClosestPointDerivative), 
+		   ClosestSplineKey(InClosestSplineKey), WaterN(InWaterN), WaterVel(FVector(0, 0, 0)), WaterVelSet(false)
+	{
+		SamplerType = EWaterSamplerType::ConstantSpline;
+		bUsePerShapeWaterPlane = false;
+		ClosestWaterPointToParticle = InClosestPoint;
+		AverageVelocity = WaterVel;
+	}
+
+	virtual void InitializeForSubmersion() 
+	{
+		WaterVel = WaterSpline.Velocity.IsSet()
+			? WaterSpline.Velocity->Eval(ClosestSplineKey) * ClosestPointDerivative.GetSafeNormal()
+			: Chaos::FVec3::ZeroVector;
+	}
+
+	virtual void SampleWaterAtPosition(const FVector& WorldPosition, FVector& OutVelocity, float& OutHeight, float& OutDepth) const override
+	{		
+		OutVelocity = WaterVel;
+		OutHeight = ClosestWaterPointToParticle.Z;
+
+		// #todo(dmp): We technically don't use depth other than needing to know it is non zero
+		// We should separate out the queries a bit more to avoid needing to set it here
+		OutDepth = 1;
+	}
+
+	virtual FVector SampleVelocityAtPosition(const FVector& WorldPosition) const override
+	{
+		return WaterVel;
+	}
+
+	virtual FVector SampleNormalAtPosition(const FVector& WorldPosition) const override
+	{
+		return WaterN;
+	}
+
+public:	
+	const FBuoyancyWaterSplineData &WaterSpline;
+	
+	FVector ClosestPointDerivative = FVector::ZeroVector;
+	float ClosestSplineKey = 0;
+
+	FVector WaterN = FVector::ZeroVector;
+	FVector WaterVel = FVector::ZeroVector;
+	bool WaterVelSet = false;
+};
 
 // Each particle will have a list of potential midphases to process,
 // which must be sorted in descending Z order. This struct is used
@@ -21,9 +147,7 @@ struct FBuoyancyInteraction
 {
 	Chaos::FPBDRigidParticleHandle* RigidParticle = nullptr;
 	Chaos::FGeometryParticleHandle* WaterParticle = nullptr;
-	const FBuoyancyWaterSplineData& WaterSpline;
-	float ClosestSplineKey = 0.f;
-	FVector ClosestPoint = FVector::ZeroVector;
+	TSharedPtr<FBuoyancyWaterSampler> WaterSampler;
 };
 
 static constexpr int32 MaxNumBuoyancyInteractions = 2;
@@ -34,10 +158,18 @@ struct FBuoyancySubmersion
 {
 	Chaos::FPBDRigidParticleHandle* Particle = nullptr;
 	TWeakPtr<FProxyTimestampBase, ESPMode::ThreadSafe> SyncTimestamp = nullptr;
+	
 	float Vol = 0.f;
 	FVector CoM = FVector::ZeroVector;
+
 	FVector Vel = FVector::ZeroVector;
-	FVector Norm = FVector::ZeroVector;	
+	FVector Norm = FVector::ZeroVector;
+
+	// #todo(dmp): only used for accurate submersion calculations for now because we integrate those per body.  Ideally we'd
+	// convert to only using this 
+	bool HasIntegratedQuantities = false;
+	FVector DeltaV = FVector::ZeroVector;
+	FVector DeltaW = FVector::ZeroVector;
 };
 
 // Metadata for submersions, used for event callbacks
diff --git a/Engine/Plugins/Experimental/Buoyancy/Source/Runtime/Public/BuoyancyStats.h b/Engine/Plugins/Experimental/Buoyancy/Source/Runtime/Public/BuoyancyStats.h
--- a/Engine/Plugins/Experimental/Buoyancy/Source/Runtime/Public/BuoyancyStats.h
+++ b/Engine/Plugins/Experimental/Buoyancy/Source/Runtime/Public/BuoyancyStats.h
@@ -26,6 +26,7 @@ DECLARE_CYCLE_STAT(TEXT("Buoyancy_Subsystem_UpdateBuoyancySettings"), STAT_Buoya
 DECLARE_CYCLE_STAT(TEXT("Buoyancy_Subsystem_UpdateWaterBodiesList"), STAT_BuoyancySubsystem_UpdateWaterBodiesList, STATGROUP_Buoyancy);
 DECLARE_CYCLE_STAT(TEXT("Buoyancy_Subsystem_DispatchCallbacks"), STAT_BuoyancySubsystem_DispatchCallbacks, STATGROUP_Buoyancy);
 DECLARE_CYCLE_STAT(TEXT("Buoyancy_Subsystem_ShallowWaterIntegrateForces"), STAT_BuoyancySubsystem_ShallowWaterIntegrateForces, STATGROUP_Buoyancy);
+DECLARE_CYCLE_STAT(TEXT("Buoyancy_Subsystem_ProcessInteractions"), STAT_BuoyancySubsystem_ProcessInteractions, STATGROUP_Buoyancy);
 DECLARE_CYCLE_STAT(TEXT("Buoyancy_Subsystem_ProcessInteraction"), STAT_BuoyancySubsystem_ProcessInteraction, STATGROUP_Buoyancy);
 
 // Algorithms
@@ -37,6 +38,11 @@ DECLARE_CYCLE_STAT(TEXT("Buoyancy_Algorithms_ComputeBuoyantForces"), STAT_Buoyan
 DECLARE_CYCLE_STAT(TEXT("Buoyancy_Algorithms_SubmergeShapeInternal"), STAT_BuoyancyAlgorithms_SubmergeShapeInternal, STATGROUP_Buoyancy);
 DECLARE_CYCLE_STAT(TEXT("Buoyancy_Algorithms_ComputeSubmergedVolumeAndForcesForParticle"), STAT_BuoyancyAlgorithms_ComputeSubmergedVolumeAndForcesForParticle, STATGROUP_Buoyancy);
 DECLARE_CYCLE_STAT(TEXT("Buoyancy_Algorithms_ComputeSubmergedVolumeAndForcesForShape"), STAT_BuoyancyAlgorithms_ComputeSubmergedVolumeAndForcesForShape, STATGROUP_Buoyancy);
+DECLARE_CYCLE_STAT(TEXT("Buoyancy_Algorithms_ComputeSubmergedVolumeAndForcesForShape_ComputeWorld"), STAT_BuoyancyAlgorithms_ComputeSubmergedVolumeAndForcesForShape_ComputeWorld, STATGROUP_Buoyancy);
+DECLARE_CYCLE_STAT(TEXT("Buoyancy_Algorithms_ComputeSubmergedVolumeAndForcesForShape_FaceIteration"), STAT_BuoyancyAlgorithms_ComputeSubmergedVolumeAndForcesForShape_FaceIteration, STATGROUP_Buoyancy);
+DECLARE_CYCLE_STAT(TEXT("Buoyancy_Algorithms_ComputeSubmergedVolumeAndForcesForShape_FindFaceIntersections"), STAT_BuoyancyAlgorithms_ComputeSubmergedVolumeAndForcesForShape_FindFaceIntersections, STATGROUP_Buoyancy);
+DECLARE_CYCLE_STAT(TEXT("Buoyancy_Algorithms_ComputeSubmergedVolumeAndForcesForShape_ComputeFaceForces"), STAT_BuoyancyAlgorithms_ComputeSubmergedVolumeAndForcesForShape_ComputeFaceForces, STATGROUP_Buoyancy);
+DECLARE_CYCLE_STAT(TEXT("Buoyancy_Algorithms_ComputeSubmergedVolumeAndForcesForShape_SampleVelocity"), STAT_BuoyancyAlgorithms_ComputeSubmergedVolumeAndForcesForShape_SampleVelocity, STATGROUP_Buoyancy);
 DECLARE_CYCLE_STAT(TEXT("Buoyancy_Algorithms_FindAllIntersectionPoints"), STAT_BuoyancyAlgorithms_FindAllIntersectionPoints, STATGROUP_Buoyancy);
 DECLARE_CYCLE_STAT(TEXT("Buoyancy_Algorithms_SortIntersectionPointsByAngle"), STAT_BuoyancyAlgorithms_SortIntersectionPointsByAngle, STATGROUP_Buoyancy);
 DECLARE_CYCLE_STAT(TEXT("Buoyancy_Algorithms_ComputeTriangleAreaAndVolume"), STAT_BuoyancyAlgorithms_ComputeTriangleAreaAndVolume, STATGROUP_Buoyancy);
diff --git a/Engine/Plugins/Experimental/Buoyancy/Source/Runtime/Public/BuoyancySubsystem.h b/Engine/Plugins/Experimental/Buoyancy/Source/Runtime/Public/BuoyancySubsystem.h
--- a/Engine/Plugins/Experimental/Buoyancy/Source/Runtime/Public/BuoyancySubsystem.h
+++ b/Engine/Plugins/Experimental/Buoyancy/Source/Runtime/Public/BuoyancySubsystem.h
@@ -229,8 +229,9 @@ private:
 	void TrackInteractions(Chaos::FPBDRigidsSolver& PBDSolver, Chaos::FPBDRigidsEvolution& Evolution, Chaos::FMidPhaseModifierAccessor& MidPhaseAccessor);
 	void TrackInteraction(Chaos::FPBDRigidsEvolution& Evolution, Chaos::FGeometryParticleHandle* WaterParticle, Chaos::FPBDRigidParticleHandle* RigidParticle, const FBuoyancyWaterSplineData& WaterSpline, Chaos::FMidPhaseModifier& MidPhase);
 	void ProcessInteractions(Chaos::FPBDRigidsEvolution& Evolution);
-	void ProcessInteractionShallowWater(Chaos::FPBDRigidsEvolution& Evolution, FBuoyancyInteraction& Interaction);
+	void ProcessAccurateInteraction(Chaos::FPBDRigidsEvolution& Evolution, FBuoyancyInteraction& Interaction, TSharedPtr<FBuoyancyWaterSampler> WaterSampler);
 	void ProcessInteraction(Chaos::FPBDRigidsEvolution& Evolution, FBuoyancyInteraction& Interaction);
+	void AddSurfaceTouchCallback(FBuoyancyInteraction& Interaction, FBuoyancySubmersion& Submersion, float TotalVol, float SubmergedVol, FVector SubmergedCoM);
 	void ApplyBuoyantForces(Chaos::FPBDRigidsEvolution& Evolution);
 	void GenerateCallbackData();
 	
