fix cis

diff --git a/Engine/Plugins/Experimental/Buoyancy/Source/Runtime/Public/BuoyancyParticleData.h b/Engine/Plugins/Experimental/Buoyancy/Source/Runtime/Public/BuoyancyParticleData.h
--- a/Engine/Plugins/Experimental/Buoyancy/Source/Runtime/Public/BuoyancyParticleData.h
+++ b/Engine/Plugins/Experimental/Buoyancy/Source/Runtime/Public/BuoyancyParticleData.h
@@ -13,7 +13,7 @@
 #include "BuoyancyWaterSplineData.h"
 #include "BuoyancyStats.h"
 
-enum EWaterSamplerType
+enum class EWaterSamplerType : uint8
 {
 	ConstantSpline,
 	ShallowWater,
