convex hulls for collisions with water

diff --git a/Engine/Plugins/Experimental/Buoyancy/Source/Runtime/Private/BuoyancyAlgorithms.cpp b/Engine/Plugins/Experimental/Buoyancy/Source/Runtime/Private/BuoyancyAlgorithms.cpp
--- a/Engine/Plugins/Experimental/Buoyancy/Source/Runtime/Private/BuoyancyAlgorithms.cpp
+++ b/Engine/Plugins/Experimental/Buoyancy/Source/Runtime/Private/BuoyancyAlgorithms.cpp
@@ -13,7 +13,6 @@
 #include "Chaos/Collision/CollisionFilter.h"
 #include "Chaos/Collision/CollisionUtil.h"
 #include "Chaos/Sphere.h"
-#include "Chaos/Convex.h"
 #include "Chaos/ImplicitObjectScaled.h"
 
 //
@@ -25,8 +24,11 @@ extern bool bBuoyancyDebugDraw;
 bool bBuoyancyAlgorithmsAllowVolRatioOverOne = false;
 FAutoConsoleVariableRef CVarBuoyancyAlgorithmsAllowVolRatioOverOne(TEXT("p.Buoyancy.Algorithms.AllowVolRatioOverOne"), bBuoyancyAlgorithmsAllowVolRatioOverOne, TEXT(""));
 
-bool bBuoyancyAlgorithmsBuoyancyAlongNormal = true;
-FAutoConsoleVariableRef CVarBuoyancyAlgorithmsBuoyancyAlongNormal(TEXT("p.Buoyancy.Algorithms.BuoyancyAlongNormal"), bBuoyancyAlgorithmsBuoyancyAlongNormal, TEXT(""));
+bool bBuoyancyAlgorithmsAccurateBuoyancyAlongNormal = false;
+FAutoConsoleVariableRef CVarBuoyancyAlgorithmsAccurateBuoyancyAlongNormal(TEXT("p.Buoyancy.Algorithms.AccurateBuoyancyAlongNormal"), bBuoyancyAlgorithmsAccurateBuoyancyAlongNormal, TEXT(""));
+
+bool bBuoyancyAlgorithmsUseConvexAsBox = false;
+FAutoConsoleVariableRef CVarBuoyancyAlgorithmsUseConvexAsBox(TEXT("p.Buoyancy.Algorithms.UseConvexAsBox"), bBuoyancyAlgorithmsUseConvexAsBox, TEXT(""));
 
 //
 // Internal Functions
@@ -177,7 +179,7 @@ namespace BuoyancyAlgorithms
 		return ParticleVol;
 	}
 
-	FRealSingle ComputeShapeVolume(const FGeometryParticleHandle* Particle)
+	FRealSingle ComputeShapeVolume(const FGeometryParticleHandle* Particle, const bool UseBoundingBoxVolume)
 	{
 		if (Particle == nullptr)
 		{
@@ -201,16 +203,23 @@ namespace BuoyancyAlgorithms
 		// ComputeSubmergedVolume.
 		FRealSingle ShapeVol = 0.f;
 		ImplicitObject->VisitLeafObjects(
-			[Particle, &ShapeInstances, &ShapeVol]
+			[Particle, UseBoundingBoxVolume , &ShapeInstances, &ShapeVol]
 			(const FImplicitObject* InnerImplicitObject, const FRigidTransform3&, const int32 RootObjectIndex, const int32, const int32)
 		{
 			const int32 ShapeIndex = ShapeInstances.IsValidIndex(RootObjectIndex) ? RootObjectIndex : 0;
 			const EImplicitObjectType ShapeType = Chaos::Private::GetImplicitCollisionType(Particle, InnerImplicitObject);
 			if (DoCollide(ShapeType, ShapeInstances[ShapeIndex].Get()))
 			{
-				Chaos::Utilities::CastHelper(*InnerImplicitObject, [&ShapeVol](const auto& Geom)
+				Chaos::Utilities::CastHelper(*InnerImplicitObject, [UseBoundingBoxVolume , &ShapeVol](const auto& Geom)
 				{
-					ShapeVol += Geom.BoundingBox().GetVolume();
+					if (UseBoundingBoxVolume)
+					{
+						ShapeVol += Geom.BoundingBox().GetVolume();
+					}
+					else
+					{
+						ShapeVol += Geom.GetVolume();
+					}
 				});
 			}
 		});
@@ -234,7 +243,7 @@ namespace BuoyancyAlgorithms
 		return ShapeVol;
 	}
 
-	void ScaleSubmergedVolume(const FPBDRigidsEvolutionGBF& Evolution, const FGeometryParticleHandle* Particle, FRealSingle& SubmergedVol, FRealSingle& TotalVol)
+	void ScaleSubmergedVolume(const FPBDRigidsEvolutionGBF& Evolution, const FGeometryParticleHandle* Particle, const bool UseBoundingBoxVolume, FRealSingle& SubmergedVol, FRealSingle& TotalVol)
 	{
 		SCOPE_CYCLE_COUNTER(STAT_BuoyancyAlgorithms_ScaleSubmergedVolume)
 
@@ -246,7 +255,7 @@ namespace BuoyancyAlgorithms
 		// The shape volume is the volume of all shape bounds which can
 		// possibly count as submerged volumes.
 		const FRealSingle ParticleVol = ComputeParticleVolume(Evolution, Particle);
-		const FRealSingle ShapeVol = ComputeShapeVolume(Particle);
+		const FRealSingle ShapeVol = ComputeShapeVolume(Particle, UseBoundingBoxVolume);
 		TotalVol = ParticleVol;
 
 		// If the submerged vol somehow exceeded the max shape vol, clamp it
@@ -278,7 +287,7 @@ namespace BuoyancyAlgorithms
 
 		if (ComputeSubmergedVolume(ParticleData, SubmergedParticle, WaterParticle, WaterX, WaterN, NumSubdivisions, MinVolume, SubmergedVol, SubmergedCoM))
 		{			
-			ScaleSubmergedVolume(Evolution, SubmergedParticle, SubmergedVol, TotalVol);
+			ScaleSubmergedVolume(Evolution, SubmergedParticle, true, SubmergedVol, TotalVol);
 
 
 #if ENABLE_DRAW_DEBUG
@@ -592,16 +601,7 @@ namespace BuoyancyAlgorithms
 
 		// Compute world buoyant force and torque
 		FVec3 WorldForce;
-
-		// optionally apply buoyancy along normal or -gravity
-		if (bBuoyancyAlgorithmsBuoyancyAlongNormal)
-		{
-			WorldForce = WaterN * BuoyantForce;
-		}
-		else
-		{
-			WorldForce = -1.0 * GravityDir * BuoyantForce;
-		}
+		WorldForce = WaterN * BuoyantForce;
 
 		const FVec3 WorldTorque = Chaos::FVec3::CrossProduct(CoMDiff, WorldForce);
 
@@ -633,10 +633,11 @@ namespace BuoyancyAlgorithms
 		//
 		return true;
 	}
-		
+
+	template <typename SamplerType>
 	void ComputeSubmergedVolumeAndForcesForParticle(FBuoyancyParticleData& ParticleData, 
 		const Chaos::FGeometryParticleHandle* SubmergedParticle, const FGeometryParticleHandle* WaterParticle, 
-		TSharedPtr<FBuoyancyWaterSampler> WaterSampler,
+		TSharedPtr<SamplerType> WaterSampler,
 		const Chaos::FPBDRigidsEvolution& Evolution, const float DeltaSeconds, const float WaterDensity, const float WaterDrag,
 		float &OutTotalParticleVol, float &OutTotalSubmergedVol, Chaos::FVec3 &OutTotalSubmergedCoM, Chaos::FVec3 &OutTotalForce, Chaos::FVec3 &OutTotalTorque)
 	{		
@@ -648,7 +649,11 @@ namespace BuoyancyAlgorithms
 		OutTotalSubmergedCoM = FVec3::ZeroVector;
 		OutTotalForce = FVec3::ZeroVector;
 		OutTotalTorque = FVec3::ZeroVector;
-			
+		
+		float OutTotalForceWeight = 0;
+		Chaos::FVec3 OutTotalBuoyancyForce = FVec3::ZeroVector;
+		Chaos::FVec3 OutTotalBuoyancyTorque = FVec3::ZeroVector;
+
 		// Get some initial data about the submerged particle
 		const FImplicitObject* RootImplicit = SubmergedParticle->GetGeometry();
 		const FShapeInstanceArray& ShapeInstances = SubmergedParticle->ShapeInstances();
@@ -668,9 +673,9 @@ namespace BuoyancyAlgorithms
 		// Traverse the submerged particle's leaves
 		RootImplicit->VisitLeafObjects(
 			[&Evolution, DeltaSeconds, WaterDensity, WaterDrag, SubmergedParticle, &WaterSampler, 
-			ParticleIndex, &ShapeInstances, WaterShapeType, WaterShapeInstance, &ParticleWorldTransform, 
+			ParticleIndex, &ShapeInstances, WaterShapeType, WaterShapeInstance, &ParticleWorldTransform,
 			&SubmergedShapes, &OutTotalSubmergedVol, &OutTotalSubmergedCoM,
-			&OutTotalForce, &OutTotalTorque]
+			&OutTotalForce, &OutTotalTorque, &OutTotalBuoyancyForce, &OutTotalBuoyancyTorque, &OutTotalForceWeight]
 			(const FImplicitObject* Implicit, const FRigidTransform3& RelativeTransform, const int32 RootObjectIndex, const int32 ObjectIndex, const int32 LeafObjectIndex)
 			{								
 				const int32 ShapeIndex = (ShapeInstances.IsValidIndex(RootObjectIndex)) ? RootObjectIndex : 0;
@@ -690,17 +695,88 @@ namespace BuoyancyAlgorithms
 					return;
 				}
 				
-				FVec3 OutWaterP, OutWaterN;
 				float OutSubmergedVol;
 				FVec3 OutSubmergedCoM;
 				FVec3 OutForce, OutTorque;
-				ComputeSubmergedVolumeAndForcesForShape(SubmergedParticle,
-					Evolution, DeltaSeconds, WaterDensity, WaterDrag,
-					Implicit, RelativeTransform, RootObjectIndex, ObjectIndex, LeafObjectIndex,
-					WaterSampler,
-					OutWaterP, OutWaterN,
-					OutSubmergedVol, OutSubmergedCoM,
-					OutForce, OutTorque);
+				FVec3 OutBuoyancyForce, OutBuoyancyTorque;
+				float OutForceWeight;
+
+				// Get the world-space bounds of shape A
+				FRigidTransform3 ShapeWorldTransform = RelativeTransform * ParticleWorldTransform;
+				FRigidTransform3 ShapeWorldTransformScaled = ShapeWorldTransform;
+
+				const FConvex* Convex = nullptr;
+				
+				FAABB3 LocalBox = Implicit->BoundingBox();
+
+				// Get the world-space bounds of shape A					
+				if (const FImplicitSphere3* Sphere = Implicit->AsA<FImplicitSphere3>())
+				{
+					// If we have a sphere, ignore rotation because submerged volume is independent of rotation
+					// and also we don't want to apply any torques on the wheel.			
+					const FVec3 SphereCenter = ShapeWorldTransform.TransformPosition(Sphere->GetCenter());
+					const FVec3 SphereExtent = FVec3(Sphere->GetRadius());
+										
+					LocalBox = FAABB3(-SphereExtent, SphereExtent);
+					ShapeWorldTransform.SetTranslation(SphereCenter);
+
+					// #todo(dmp): at this point, we haven't figured out the water normal so we use up vector for now
+					FVec3 TmpWaterN = { 0.f, 0.f, 1.f };
+					ShapeWorldTransform.SetRotation(FRotationMatrix::MakeFromZ(TmpWaterN).ToQuat());
+				}
+				else if (const FConvex* TheConvex = Implicit->AsA<FConvex>())
+				{
+					Convex = TheConvex;
+				}
+				else if (const TImplicitObjectScaled<FConvex>* ScaledObject = TImplicitObjectScaled<FConvex>::AsScaled(*Implicit))
+				{
+					Convex = ScaledObject->GetUnscaledObject();
+
+					if (Convex != nullptr)
+					{
+						FRigidTransform3 RelativeTransformTmp = RelativeTransform;
+						RelativeTransformTmp.SetScale3D(RelativeTransform.GetScale3D() * ScaledObject->GetScale());
+
+						ShapeWorldTransformScaled = RelativeTransformTmp * ParticleWorldTransform;
+					}
+				}
+				else if (const TImplicitObjectInstanced<FConvex>* InstancedObject = TImplicitObjectInstanced<FConvex>::AsInstanced(*Implicit))
+				{
+					Convex = InstancedObject->Object();
+				}				
+
+				// We support convex and boxes separately with the same buoyancy algorithm
+				if (Convex && !bBuoyancyAlgorithmsUseConvexAsBox)
+				{
+					FBuoyancyConvexShape ConvexShapeQuery(Convex);
+					ConvexShapeQuery.Initialize();
+
+					ComputeSubmergedVolumeAndForcesForShape<FBuoyancyConvexShape, SamplerType>(SubmergedParticle, ConvexShapeQuery,
+						Evolution, DeltaSeconds, WaterDensity, WaterDrag,
+						ShapeWorldTransformScaled, WaterSampler,
+						OutSubmergedVol, OutSubmergedCoM, OutForce, OutTorque, OutBuoyancyForce, OutBuoyancyTorque, OutForceWeight);
+				}
+				else
+				{					
+					FBuoyancyBoxShape BoxShapeQuery(LocalBox);
+					BoxShapeQuery.Initialize();
+
+					ComputeSubmergedVolumeAndForcesForShape<FBuoyancyBoxShape, SamplerType>(SubmergedParticle, BoxShapeQuery,
+						Evolution, DeltaSeconds, WaterDensity, WaterDrag,
+						ShapeWorldTransform, WaterSampler,
+						OutSubmergedVol, OutSubmergedCoM, OutForce, OutTorque, OutBuoyancyForce, OutBuoyancyTorque, OutForceWeight);
+				}
+
+#if ENABLE_DRAW_DEBUG
+				if (bBuoyancyDebugDraw)
+				{					
+					Chaos::FDebugDrawQueue::GetInstance().DrawDebugBox(
+						ShapeWorldTransform.TransformPosition(LocalBox.GetCenter()),
+						LocalBox.Extents() * .5f,
+						ShapeWorldTransform.GetRotation(),
+						FColor::Red, false, -1.f, SDPG_Foreground, 3.f);
+				}
+#endif
 
 				// mark the shape as submerged
 				if (OutSubmergedVol > SMALL_NUMBER)
@@ -711,52 +787,58 @@ namespace BuoyancyAlgorithms
 
 					OutTotalSubmergedVol += OutSubmergedVol;
 
+					OutTotalForceWeight += OutForceWeight;
+
+					OutTotalBuoyancyForce += OutBuoyancyForce;
+					OutTotalBuoyancyTorque += OutBuoyancyTorque;
+
 					SubmergeShape_Internal(SubmergedShapes, ParticleIndex, ObjectIndex);
 				}
 			});
 
+		OutTotalForce /= OutTotalForceWeight;
+		OutTotalTorque /= OutTotalForceWeight;
+
+		OutTotalForce += OutTotalBuoyancyForce;
+		OutTotalTorque += OutTotalBuoyancyTorque;
+
 		// compute final force and torque for particle as weighted average of all the shapes
 		OutTotalSubmergedCoM /= OutTotalSubmergedVol;
 
 		float ScaledSubmergedVol = OutTotalSubmergedVol;
-		ScaleSubmergedVolume(Evolution, SubmergedParticle, ScaledSubmergedVol, OutTotalParticleVol);
+		ScaleSubmergedVolume(Evolution, SubmergedParticle, true, ScaledSubmergedVol, OutTotalParticleVol);
 	}
 
-	void ComputeSubmergedVolumeAndForcesForShape(const Chaos::FGeometryParticleHandle* SubmergedParticle,
+	// forward declaration of usage	
+	template void ComputeSubmergedVolumeAndForcesForParticle<FBuoyancyConstantSplineSampler>(FBuoyancyParticleData&,
+		const Chaos::FGeometryParticleHandle*, const Chaos::FGeometryParticleHandle*,
+		TSharedPtr<FBuoyancyConstantSplineSampler>,
+		const Chaos::FPBDRigidsEvolution&, const float, const float, const float,
+		float&, float&, Chaos::FVec3&, Chaos::FVec3&, Chaos::FVec3&);
+
+	template void ComputeSubmergedVolumeAndForcesForParticle<FBuoyancyShallowWaterSampler>(FBuoyancyParticleData&,
+		const Chaos::FGeometryParticleHandle*, const Chaos::FGeometryParticleHandle*,
+		TSharedPtr<FBuoyancyShallowWaterSampler>,
+		const Chaos::FPBDRigidsEvolution&, const float, const float, const float,
+		float&, float&, Chaos::FVec3&, Chaos::FVec3&, Chaos::FVec3&);
+
+	template <typename ShapeType, typename SamplerType>
+	void ComputeSubmergedVolumeAndForcesForShape(const Chaos::FGeometryParticleHandle* SubmergedParticle, const ShapeType& ShapeQuery,
 		const Chaos::FPBDRigidsEvolution& Evolution, float DeltaSeconds, const float WaterDensity, const float WaterDrag,
-		const Chaos::FImplicitObject* Implicit, const Chaos::FRigidTransform3& RelativeTransform, const int32 RootObjectIndex, const int32 ObjectIndex, const int32 LeafObjectIndex,
-		TSharedPtr<FBuoyancyWaterSampler> WaterSampler,
-		Chaos::FVec3& OutWaterP, Chaos::FVec3& OutWaterN,
+		const Chaos::FRigidTransform3 &ShapeWorldTransform, TSharedPtr<SamplerType> WaterSampler,
 		float& OutSubmergedVol, Chaos::FVec3& OutSubmergedCoM,
-		Chaos::FVec3& OutForce, Chaos::FVec3& OutTorque)
+		Chaos::FVec3& OutForce, Chaos::FVec3& OutTorque,
+		Chaos::FVec3& OutBuoyancyForce, Chaos::FVec3& OutBuoyancyTorque,
+		float& OutForceWeight)
 	{
 		SCOPE_CYCLE_COUNTER(STAT_BuoyancyAlgorithms_ComputeSubmergedVolumeAndForcesForShape)
 
-		const FConstGenericParticleHandle SubmergedGeneric = SubmergedParticle;
-		const FRigidTransform3 ParticleWorldTransform = SubmergedGeneric->GetTransformPQ();
-
-		const FPBDRigidParticleHandle* RigidParticle = SubmergedParticle->CastToRigidParticle();
+		const FConstGenericParticleHandle SubmergedGeneric = SubmergedParticle;		
+		const FPBDRigidParticleHandle* RigidParticle = SubmergedParticle->CastToRigidParticle();					
 
-		// Get the world-space bounds of shape A
-		FRigidTransform3 ShapeWorldTransform = RelativeTransform * ParticleWorldTransform;
-		FAABB3 LocalBox = Implicit->BoundingBox();
-		if (const FImplicitSphere3* Sphere = Implicit->AsA<FImplicitSphere3>())
-		{
-			// If we have a sphere, ignore rotation because submerged volume is independent of rotation
-			// and also we don't want to apply any torques on the wheel.			
-			const FVec3 SphereCenter = ShapeWorldTransform.TransformPosition(Sphere->GetCenter());
-			const FVec3 SphereExtent = FVec3(Sphere->GetRadius());
-			LocalBox = FAABB3(-SphereExtent, SphereExtent);
-			ShapeWorldTransform.SetTranslation(SphereCenter);
-
-			// #todo(dmp): at this point, we haven't figured out the water normal so we use up vector for now
-			FVec3 TmpWaterN = {0.f, 0.f, 1.f };
-			ShapeWorldTransform.SetRotation(FRotationMatrix::MakeFromZ(TmpWaterN).ToQuat());
-		}
-		
 		// water plane parameters for the body.  Note we have 1 plane per box
-		OutWaterP = FVec3::ZeroVector;
-		OutWaterN = FVec3(0.f, 0.f, 1.f);
+		Chaos::FVec3 ShapeWaterP = FVec3::ZeroVector;
+		Chaos::FVec3 ShapeWaterN = FVec3(0.f, 0.f, 1.f);
 
 		// Compute submerged volume and CoM
 		OutSubmergedVol = 0.f;
@@ -771,7 +853,7 @@ namespace BuoyancyAlgorithms
 		// either use a water plane per shape's centroid or the closest water point to the particle
 		if (WaterSampler->UsePerShapeWaterPlane())
 		{			
-			const FVector WorldBoxCenter = ShapeWorldTransform.TransformPosition(LocalBox.GetCenter());
+			const FVector WorldBoxCenter = ShapeWorldTransform.TransformPosition(ShapeQuery.GetCenter());
 
 			// water plane position is computed from the centroid of the box		
 			FVector CenterVelocity;
@@ -785,27 +867,25 @@ namespace BuoyancyAlgorithms
 				return;
 			}
 						
-			OutWaterP = WorldBoxCenter;
-			OutWaterP.Z = CenterHeight;
+			ShapeWaterP = WorldBoxCenter;
+			ShapeWaterP.Z = CenterHeight;
 		}
 		else
 		{
-			OutWaterP = WaterSampler->GetClosestWaterPointToParticle();
+			ShapeWaterP = WaterSampler->GetClosestWaterPointToParticle();
 		}
 
 		// Sample normal at water plane position
 		// #todo(dmp): just store normals in the shallow water texture rather than finite differencing here
-		OutWaterN = WaterSampler->SampleNormalAtPosition(OutWaterP);
+		ShapeWaterN = WaterSampler->SampleNormalAtPosition(ShapeWaterP);
+								
+		TArray<FVector, TInlineAllocator<FBuoyancyShapeTopologyLimits::MaxVerticesPerShape>> WorldVertexPosition;
+		TArray<bool, TInlineAllocator<FBuoyancyShapeTopologyLimits::MaxVerticesPerShape>> VertexIsUnderwater;
 
-		// #todo(dmp): generalize to convex shapes and remove box assumption
-		static constexpr int32 TotalNumVertices = 8;
-		static constexpr int32 NumFaces = 6;
-		static constexpr int32 NumEdges = 12;
-		static constexpr int32 NumVerticesForFace = 4;
-		static constexpr int32 MaxBoxPlaneIntersectionsPoints = 6;
+		const int32 TotalNumVertices = ShapeQuery.NumVertices();
 
-		FVector WorldVertexPosition[TotalNumVertices];
-		bool VertexIsUnderwater[TotalNumVertices];
+		WorldVertexPosition.AddUninitialized(TotalNumVertices);
+		VertexIsUnderwater.AddUninitialized(TotalNumVertices);
 
 		// keep track of a reference point inside the clipped geometry that is the average of all the 
 		// submerged vertices and intersection points this will jump around because it isn't weighted by the area/volume
@@ -819,10 +899,10 @@ namespace BuoyancyAlgorithms
 			// SCOPE_CYCLE_COUNTER(STAT_BuoyancyAlgorithms_ComputeSubmergedVolumeAndForcesForShape_ComputeWorld);
 
 			// compute world space position of vertex
-			const FVector CurrWorldVertexPosition = ShapeWorldTransform.TransformPosition(LocalBox.GetVertex(i));
+			const FVector CurrWorldVertexPosition = ShapeWorldTransform.TransformPosition(ShapeQuery.GetVertex(i));
 			WorldVertexPosition[i] = CurrWorldVertexPosition;
 
-			if ((CurrWorldVertexPosition - OutWaterP).Dot(OutWaterN) < SMALL_NUMBER)
+			if ((CurrWorldVertexPosition - ShapeWaterP).Dot(ShapeWaterN) < SMALL_NUMBER)
 			{
 				VertexIsUnderwater[i] = true;
 				InteriorRefPoint += WorldVertexPosition[i];
@@ -844,87 +924,90 @@ namespace BuoyancyAlgorithms
 		FVector IntersectionCenter;			
 
 		// we can have at most 6 intersection points with a box and a plane			
-		TArray<FVector, TFixedAllocator<MaxBoxPlaneIntersectionsPoints>> OrderedIntersectionPoints;
+		TArray<FVector, TInlineAllocator<FBuoyancyShapeTopologyLimits::MaxIntersectionPointsPerShape>> OrderedIntersectionPoints;		
+		TMap<int32, FVector> EdgeToIntersectionPoint;
+
+		const int32 NumEdges = ShapeQuery.NumEdges();
+		const int32 MaxBoxPlaneIntersectionsPoints = ShapeQuery.NumMaxIntersectionsPoints();
+
+		OrderedIntersectionPoints.AddUninitialized(MaxBoxPlaneIntersectionsPoints);		
 
-		FVector EdgeToIntersectionPoint[NumEdges];
 		int32 NumIntersectionPoints;
-		BuoyancyAlgorithms::FindAllIntersectionPoints(OutWaterP, OutWaterN, LocalBox, WorldVertexPosition, EdgeToIntersectionPoint, NumIntersectionPoints, 
-			OrderedIntersectionPoints, IntersectionCenter);
-		
+		BuoyancyAlgorithms::FindAllIntersectionPoints<ShapeType>(ShapeWaterP, ShapeWaterN, ShapeQuery, WorldVertexPosition, EdgeToIntersectionPoint, NumIntersectionPoints,
+			OrderedIntersectionPoints, IntersectionCenter);		
+
 		// Add intersection points to interior reference point
 		// values are already accumulated into the intersection center
 		InteriorRefPoint += IntersectionCenter * NumIntersectionPoints;
 		InteriorRefPointCount += NumIntersectionPoints;
 		InteriorRefPoint /= InteriorRefPointCount;
 
-		float AccumulatedForceWeight = 0;
+		OutForceWeight = 0;
 		
 		const FVec3 WorldCoM = SubmergedGeneric->PCom();
 
 		// Sum up submerged areas and volumes for each face
-		
+		const int32 NumFaces = ShapeQuery.NumFaces();
 		for (int32 FaceIdx = 0; FaceIdx < NumFaces; ++FaceIdx)
 		{
 			// SCOPE_CYCLE_COUNTER(STAT_BuoyancyAlgorithms_ComputeSubmergedVolumeAndForcesForShape_FaceIteration);
+			
+			FVector SubmergedFaceCenter = FVector::ZeroVector;
 
-			const Chaos::FAABBFace CurrFace = LocalBox.GetFace(FaceIdx);
-
-			FVector IntersectedFaceCenter = FVector::ZeroVector;
-
-			// we know that a given face can have at most 5 points below the surface
-			// #todo(dmp): we need to remove this assumption if we support convex hulls or other types of geometry
-			static constexpr int32 MaxNumVerticesBelowWater = 5;
-			TArray<FVector, TFixedAllocator<MaxNumVerticesBelowWater>> FaceIntersectionPoints;
+			// we know that a given face can have at most one more than the max number of vertices per face			
+			
+			TArray<FVector, TInlineAllocator<FBuoyancyShapeTopologyLimits::MaxSubmergedFaceVertices>> SubmergedFaceVertices;
 
 			// walk vertices for the current face in counter clockwise order and find intersections and
-			// submerged vertices			
+			// submerged vertices		
+			const int32 NumVerticesForFace = ShapeQuery.NumFaceVertices(FaceIdx);
+
 			for (int32 i = 0; i < NumVerticesForFace; ++i)
 			{
-				// SCOPE_CYCLE_COUNTER(STAT_BuoyancyAlgorithms_ComputeSubmergedVolumeAndForcesForShape_FindFaceIntersections);
+				// SCOPE_CYCLE_COUNTER(STAT_BuoyancyAlgorithms_ComputeSubmergedVolumeAndForcesForShape_FindFaceVerticesBelowWater);
 
 				// get the current edge index belonging to the i and i+1 vertices
 				// note this is a bit weird...we could just ask the edge for the vertices but they
 				// might be in the wrong order for correctly constructing the intersected face
-				const int32 CurrEdgeIdx = CurrFace.EdgeIndex[i];
-			
-				// if there is an intersection for this edge, add it to the list for the current face
-				const FVector &EdgeIntersection = EdgeToIntersectionPoint[CurrEdgeIdx];
+				const int32 CurrEdgeIdx = ShapeQuery.GetFaceEdge(FaceIdx, i);
 
-				// #todo(dmp): use bit mask for this
-				if (EdgeIntersection.X < TNumericLimits<float>::Max())
+				// if there is an intersection for this edge, add it to the list for the current face
+				const FVector *EdgeIntersection = EdgeToIntersectionPoint.Find(CurrEdgeIdx);
+				
+				if (EdgeIntersection)
 				{
-					FaceIntersectionPoints.Add(EdgeIntersection);
-					IntersectedFaceCenter += EdgeIntersection;
+					SubmergedFaceVertices.Add(*EdgeIntersection);
+					SubmergedFaceCenter += *EdgeIntersection;
 				}
 
 				// Endpoint of the edge
-				const int32 Edge1Idx = CurrFace.VertexIndex[(i + 1) % 4];
+				const int32 Edge1Idx = ShapeQuery.GetFaceVertex(FaceIdx, (i + 1) % NumVerticesForFace);
 
 				// add the second endpoint if it is underwater
 				if (VertexIsUnderwater[Edge1Idx])
 				{
 					const FVector& V1 = WorldVertexPosition[Edge1Idx];
 
-					FaceIntersectionPoints.Add(V1);
-					IntersectedFaceCenter += V1;
+					SubmergedFaceVertices.Add(V1);
+					SubmergedFaceCenter += V1;
 				}
 			}
 			
-			const int32 NumFaceIntersectionPts = FaceIntersectionPoints.Num();
+			const int32 NumSubmergedFaceVertices = SubmergedFaceVertices.Num();
 
 			// if the face intersects the water or is fully submerged, compute face coverage
-			if (NumFaceIntersectionPts > 0)
+			if (NumSubmergedFaceVertices > 0)
 			{
-				IntersectedFaceCenter /= NumFaceIntersectionPts;
+				SubmergedFaceCenter /= NumSubmergedFaceVertices;
 
 				// we have a list of the intersection points for a face in correct counterclockwise order
 				// fan layout for points on face, sum up areas and volume for the current face
-				for (int32 i = 0; i < NumFaceIntersectionPts; i++)
+				for (int32 i = 0; i < NumSubmergedFaceVertices; i++)
 				{
 					// SCOPE_CYCLE_COUNTER(STAT_BuoyancyAlgorithms_ComputeSubmergedVolumeAndForcesForShape_ComputeFaceForces);
-					const FVector &V0 = FaceIntersectionPoints[i];
-					const FVector &V1 = FaceIntersectionPoints[(i + 1) % NumFaceIntersectionPts];
-					const FVector &V2 = IntersectedFaceCenter;
+					const FVector &V0 = SubmergedFaceVertices[i];
+					const FVector &V1 = SubmergedFaceVertices[(i + 1) % NumSubmergedFaceVertices];
+					const FVector &V2 = SubmergedFaceCenter;
 
 					FVector TriBaryCenter;
 					float TriArea;
@@ -960,14 +1043,14 @@ namespace BuoyancyAlgorithms
 					FVector TotalWorldTorque;
 
 					// #todo(dmp): we scale drag to make it closer to the old algorithm for now
-					const float DragFactor = .1 * WaterDrag;
+					const float DragFactor = .05 * WaterDrag;
 					BuoyancyAlgorithms::ComputeFluidForceForTriangle(
 						DragFactor, DeltaSeconds, RigidParticle, WorldCoM, TriBaryCenter, TriNormal, TriArea, TetVolume,
-						CurrWaterVelocity, OutWaterP, OutWaterN, TotalWorldForce, TotalWorldTorque);
+						CurrWaterVelocity, ShapeWaterP, ShapeWaterN, TotalWorldForce, TotalWorldTorque);
 
 					OutForce += TotalWorldForce * TriArea;
 					OutTorque += TotalWorldTorque * TriArea;
-					AccumulatedForceWeight += TriArea;
+					OutForceWeight += TriArea;
 
 #if ENABLE_DRAW_DEBUG
 					if (bBuoyancyDebugDraw)
@@ -975,18 +1058,21 @@ namespace BuoyancyAlgorithms
 						FVector CurrWaterVelocityViz = CurrWaterVelocity;
 						CurrWaterVelocityViz.Normalize();
 						Chaos::FDebugDrawQueue::GetInstance().DrawDebugDirectionalArrow(TriBaryCenter, TriBaryCenter + CurrWaterVelocityViz * 20, 20, FColor::Red, false, -1.f, -1, 2.f);
-						Chaos::FDebugDrawQueue::GetInstance().DrawDebugDirectionalArrow(TriBaryCenter, TriBaryCenter + TotalWorldForce * .01, 20, FColor::Green, false, -1.f, -1, 2.f);
+						Chaos::FDebugDrawQueue::GetInstance().DrawDebugDirectionalArrow(TriBaryCenter, TriBaryCenter + TotalWorldForce * .000025, 20, FColor::Green, false, -1.f, -1, 2.f);
 					}
 #endif					
 				}
 			}
 		}
 
-		// Sum up tet volumes from the fan of triangles across the intersection face		
-		for (int32 i = 0; i < NumIntersectionPoints; i++)
+		// Sum up tet volumes from the fan of triangles across the intersection face	
+		// note that the unique number of intersection points might differ as some convex hull points can collapse to a single one during the ordering
+		// process of the intersection algorithm	
+		int32 NumUniqueIntersectionPoints = OrderedIntersectionPoints.Num();
+		for (int32 i = 0; i < NumUniqueIntersectionPoints; i++)
 		{
 			const FVector &V0 = OrderedIntersectionPoints[i];
-			const FVector &V1 = OrderedIntersectionPoints[(i + 1) % NumIntersectionPoints];
+			const FVector &V1 = OrderedIntersectionPoints[(i + 1) % NumUniqueIntersectionPoints];
 			const FVector &V2 = IntersectionCenter;
 
 			FVector TriBaryCenter;
@@ -1008,108 +1094,94 @@ namespace BuoyancyAlgorithms
 		}
 					
 		OutSubmergedCoM /= SubmergedCOMTotalWeight;
-				
-		// #todo(dmp): is it correct to normalize by the accumulated area over which the force is done?  Seems wrong...
-		if (AccumulatedForceWeight > SMALL_NUMBER)
-		{
-			OutForce /= AccumulatedForceWeight;
-			OutTorque /= AccumulatedForceWeight;
-		}
 
 		// compute buoyancy force and delta velocity for solver
 		Chaos::FVec3 WorldBuoyantForce, WorldBuoyantTorque;
 		BuoyancyAlgorithms::ComputeBuoyantForceForShape(Evolution, RigidParticle, DeltaSeconds,
-			WaterDensity, OutSubmergedCoM, OutSubmergedVol, OutWaterN, WorldBuoyantForce, WorldBuoyantTorque);
-
-		// Total force acting on the particle is the buoyant force plus all the forces due to matching the fluid velocity
-		OutForce += WorldBuoyantForce;
-		OutTorque += WorldBuoyantTorque;
+			WaterDensity, OutSubmergedCoM, OutSubmergedVol, ShapeWaterN, WorldBuoyantForce, WorldBuoyantTorque);
+		
+		// track the buoyancy force separately
+		OutBuoyancyForce = WorldBuoyantForce;
+		OutBuoyancyTorque = WorldBuoyantTorque;
 
 #if ENABLE_DRAW_DEBUG
 		if (bBuoyancyDebugDraw)
 		{
-			FVector WorldBoxCenter = ShapeWorldTransform.TransformPosition(LocalBox.GetCenter());
+			FVector WorldBoxCenter = ShapeWorldTransform.TransformPosition(ShapeQuery.GetCenter());
 
+			Chaos::FDebugDrawQueue::GetInstance().DrawDebugDirectionalArrow(OutSubmergedCoM, OutSubmergedCoM + OutForce * .00005, 20, FColor::Black, false, -1.f, -1, 2.f);
 			Chaos::FDebugDrawQueue::GetInstance().DrawDebugSphere(WorldBoxCenter, 10, 10, FColor::Green, false, -1.f, -1, 2.0f);
 			Chaos::FDebugDrawQueue::GetInstance().DrawDebugSphere(InteriorRefPoint, 10, 10, FColor::Yellow, false, -1.f, -1, 2.0f);
-			Chaos::FDebugDrawQueue::GetInstance().DrawDebugSphere(OutWaterP, 10, 10, FColor::Orange, false, -1.f, -1, 2.0f);
-			Chaos::FDebugDrawQueue::GetInstance().DrawDebugDirectionalArrow(OutWaterP, OutWaterP + OutWaterN * 100, 20, FColor::Orange, false, -1.f, -1, 6.f);
+			Chaos::FDebugDrawQueue::GetInstance().DrawDebugSphere(ShapeWaterP, 10, 10, FColor::Orange, false, -1.f, -1, 2.0f);
+			Chaos::FDebugDrawQueue::GetInstance().DrawDebugDirectionalArrow(ShapeWaterP, ShapeWaterP + ShapeWaterN * 100, 20, FColor::Orange, false, -1.f, -1, 6.f);
 			Chaos::FDebugDrawQueue::GetInstance().DrawDebugSphere(OutSubmergedCoM, 10, 10, FColor::Magenta, false, -1.f, -1, 2.0f);
-
-			//float SubmergedVol1;
-			//Chaos::FVec3 SubmergedCoM1;
-			//float TotalVol1;
-			//if (BuoyancyAlgorithms::ComputeSubmergedVolume(BuoyancyParticleData, Evolution, Interaction.RigidParticle, Interaction.WaterParticle, Interaction.ClosestPoint, WaterN, BuoyancySettings->MaxNumBoundsSubdivisions, BuoyancySettings->MinBoundsSubdivisionVol, SubmergedVol1, SubmergedCoM1, TotalVol1))
-			//{
-			//}
-			//Chaos::FDebugDrawQueue::GetInstance().DrawDebugSphere(SubmergedCoM1, 10, 10, FColor::Black, false, -1.f, -1, 2.0f);
-
-			Chaos::FDebugDrawQueue::GetInstance().DrawDebugBox(
-				ShapeWorldTransform.TransformPosition(LocalBox.GetCenter()),
-				LocalBox.Extents() * .5f,
-				ShapeWorldTransform.GetRotation(),
-				FColor::Red, false, -1.f, SDPG_Foreground, 3.f);
+			Chaos::FDebugDrawQueue::GetInstance().DrawDebugDirectionalArrow(OutSubmergedCoM, OutSubmergedCoM + OutForce * .00005, 20, FColor::Magenta, false, -1.f, -1, 2.f);
+			Chaos::FDebugDrawQueue::GetInstance().DrawDebugDirectionalArrow(OutSubmergedCoM, OutSubmergedCoM + OutTorque * .00005, 20, FColor::Red, false, -1.f, -1, 2.f);
 		}
-#endif	
+#endif
 	}
 
-	void FindAllIntersectionPoints(const Chaos::FVec3& WaterP, const Chaos::FVec3& WaterN, const Chaos::FAABB3& LocalBox, const FVector WorldVertexPosition[8],
-		FVector EdgeToIntersectionPoint[12], int32 &NumIntersections, TArray<FVector, TFixedAllocator<6>>& OutOrderedIntersectionPoints, FVector& OutIntersectionCenter)
+	template <typename ShapeType>
+	void FindAllIntersectionPoints(const Chaos::FVec3& WaterP, const Chaos::FVec3& WaterN, const ShapeType&ShapeQuery, 
+		const TArray<FVector, TInlineAllocator<FBuoyancyShapeTopologyLimits::MaxVerticesPerShape>> &WorldVertexPosition,
+		TMap<int32, FVector> &EdgeToIntersectionPoint, int32 &NumIntersections,
+		TArray<FVector, TInlineAllocator<FBuoyancyShapeTopologyLimits::MaxIntersectionPointsPerShape>>& OutOrderedIntersectionPoints, FVector& OutIntersectionCenter)
 	{		
 		// SCOPE_CYCLE_COUNTER(STAT_BuoyancyAlgorithms_FindAllIntersectionPoints)
 
 		OutIntersectionCenter = FVector::ZeroVector;
-		
-		NumIntersections = 0;
 
 		// loop over all edges and compute intersection points
-		for (int32 CurrEdgeIdx = 0; CurrEdgeIdx < 12; ++CurrEdgeIdx)
-		{
-			const FAABBEdge CurrEdge = LocalBox.GetEdge(CurrEdgeIdx);
+		const int32 NumEdges = ShapeQuery.NumEdges();
+		for (int32 CurrEdgeIdx = 0; CurrEdgeIdx < NumEdges; ++CurrEdgeIdx)
+		{		
+			int32 Index0, Index1;
+			ShapeQuery.GetEdgeVertices(CurrEdgeIdx, Index0, Index1);
 
 			Chaos::FVec3 IntersectionPoint;
-			if (EdgePlaneIntersection(WaterP, WaterN, WorldVertexPosition[CurrEdge.VertexIndex0], WorldVertexPosition[CurrEdge.VertexIndex1], IntersectionPoint))
+			if (EdgePlaneIntersection(WaterP, WaterN, WorldVertexPosition[Index0], WorldVertexPosition[Index1], IntersectionPoint))
 			{
-				OutIntersectionCenter += IntersectionPoint;
-				NumIntersections++;
+				EdgeToIntersectionPoint.Add(CurrEdgeIdx, IntersectionPoint);
+				OutIntersectionCenter += IntersectionPoint;				
 			}
-
-			EdgeToIntersectionPoint[CurrEdgeIdx] = IntersectionPoint;
 		}
 
+		NumIntersections = EdgeToIntersectionPoint.Num();
+
 		if (NumIntersections > 0)
 		{
 			OutIntersectionCenter /= NumIntersections;
 
 			// Sort intersection points by angle around reference point
-			BuoyancyAlgorithms::SortIntersectionPointsByAngle(WaterP, WaterN, OutIntersectionCenter, EdgeToIntersectionPoint, NumIntersections, OutOrderedIntersectionPoints);
+			BuoyancyAlgorithms::SortIntersectionPointsByAngle<ShapeType>(WaterP, WaterN, OutIntersectionCenter, ShapeQuery, EdgeToIntersectionPoint, OutOrderedIntersectionPoints);			
+		}
+		else
+		{
+			OutOrderedIntersectionPoints.Reset(0);
 		}
 	}
 
-	void SortIntersectionPointsByAngle(const Chaos::FVec3& WaterP, const Chaos::FVec3& WaterN, const Chaos::FVec3& IntersectionCenter, 
-		const FVector EdgeToIntersectionPoint[12], const int32 NumIntersections, TArray<FVector, TFixedAllocator<6>>& OutOrderedIntersectionPoints)
+	template <typename ShapeType>
+	void SortIntersectionPointsByAngle(const Chaos::FVec3& WaterP, const Chaos::FVec3& WaterN, const Chaos::FVec3& IntersectionCenter, const ShapeType& ShapeQuery,
+		const TMap<int, FVector>& EdgeToIntersectionPoint, TArray<FVector, TInlineAllocator<FBuoyancyShapeTopologyLimits::MaxIntersectionPointsPerShape>>& OutOrderedIntersectionPoints)
 	{
 		SCOPE_CYCLE_COUNTER(STAT_BuoyancyAlgorithms_SortIntersectionPointsByAngle)
 
 		// Order intersection points around intersection center
-		TMap<float, FVector> IntersectionPointAngleToCenter;
+		TMap<double, FVector> IntersectionPointAngleToCenter;
 				
 		bool FoundFirst = false;
 
 		// first element is assigned to be angle 0
 		FVector AngleStart;		
 		
-		// loop over all points, determine angle to reference point
-		for (int32 i = 0; i < 12; ++i)
+		// loop over all edges, find intersection points then determine angle to reference point		
+		for (const TPair<int32, FVector>& IntersectedEdge : EdgeToIntersectionPoint)
 		{
-			float CurrAngle;
-			const FVector &CurrPoint = EdgeToIntersectionPoint[i];
+			double CurrAngle;
+			const FVector& CurrPoint = IntersectedEdge.Value;
 			
-			if (CurrPoint.X >= TNumericLimits<float>::Max()) // not a valid point
-			{
-				continue;
-			}						
-			else if (!FoundFirst) // first point found is the 0 angle point
+			if (!FoundFirst) // first point found is the 0 angle point
 			{
 				AngleStart = CurrPoint - IntersectionCenter;
 				AngleStart.Normalize();
@@ -1133,10 +1205,9 @@ namespace BuoyancyAlgorithms
 			IntersectionPointAngleToCenter.Add(CurrAngle, CurrPoint);
 		}
 
-		// #todo(dmp): replace this sort with something faster since we only have between 3-6 points to sort (and index 0 is already sorted)
-		// order intersection points by the angle it makes so we have a correct winding order
-		// for tesselation
-		IntersectionPointAngleToCenter.KeySort([](const float& A, const float& B) {
+		// #todo(dmp): replace this sort with something faster since we only have between 3-6 points to sort (and index 0 is already sorted) for boxes
+		// order intersection points by the angle it makes so we have a correct winding order for tesselation
+		IntersectionPointAngleToCenter.KeySort([](const double& A, const double& B) {
 			return A < B;
 			});
 		
@@ -1202,8 +1273,8 @@ namespace BuoyancyAlgorithms
 			Chaos::FDebugDrawQueue::GetInstance().DrawDebugLine(V1, V2, FColor::Orange, false, -1.f, 0, 1.f);
 			Chaos::FDebugDrawQueue::GetInstance().DrawDebugLine(V2, V0, FColor::Orange, false, -1.f, 0, 1.f);
 
-			Chaos::FDebugDrawQueue::GetInstance().DrawDebugSphere(OutTriangleBaryCenter, 5, 10, FColor::Cyan, false, -1.f, -1, 1.f);
-			Chaos::FDebugDrawQueue::GetInstance().DrawDebugDirectionalArrow(OutTriangleBaryCenter, OutTriangleBaryCenter + OutNormal * 20, 20, FColor::Cyan, false, -1.f, -1, 2.f);
+			//Chaos::FDebugDrawQueue::GetInstance().DrawDebugSphere(OutTriangleBaryCenter, 5, 10, FColor::Cyan, false, -1.f, -1, 1.f);
+			//Chaos::FDebugDrawQueue::GetInstance().DrawDebugDirectionalArrow(OutTriangleBaryCenter, OutTriangleBaryCenter + OutNormal * 20, 20, FColor::Cyan, false, -1.f, -1, 2.f);
 		}
 #endif
 	}
@@ -1291,7 +1362,7 @@ namespace BuoyancyAlgorithms
 		// Scale submerged volume
 		float ScaledSubmergedVol = SubmergedVol;
 		float TotalParticleVol;
-		ScaleSubmergedVolume(Evolution, RigidParticle, ScaledSubmergedVol, TotalParticleVol);
+		ScaleSubmergedVolume(Evolution, RigidParticle, false, ScaledSubmergedVol, TotalParticleVol);
 
 		// Compute buoyant force
 		//
@@ -1314,9 +1385,8 @@ namespace BuoyancyAlgorithms
 
 		const FVec3 WorldCoM = RigidGeneric->PCom();
 		const FVec3 CoMDiff = SubmergedCoM - WorldCoM;
-
-		// optionally apply buoyancy along normal or -gravity
-		if (bBuoyancyAlgorithmsBuoyancyAlongNormal)
+		
+		if (bBuoyancyAlgorithmsAccurateBuoyancyAlongNormal)
 		{
 			OutWorldBuoyantForce = WaterN * BuoyantForceMagnitude;
 		}
@@ -1329,4 +1399,21 @@ namespace BuoyancyAlgorithms
 		// buoyancy is computed at the submerged center of mass		
 		OutWorldBuoyantTorque = Chaos::FVec3::CrossProduct(CoMDiff, OutWorldBuoyantForce);				
 	}
-}
\ No newline at end of file
+
+	void FBuoyancyConvexShape::Initialize()
+	{		
+		// #todo(dmp): ideally the convex would have a mapping between half edges and edges that is serialized, but instead for now we'll build it here
+		const int32 NumEdges = Convex->NumEdges();
+		HalfEdgeToEdge.AddUninitialized(NumEdges * 2);
+
+		for (int32 CurrEdgeIdx = 0; CurrEdgeIdx < NumEdges; ++CurrEdgeIdx)
+		{
+			int32 HalfEdge0, HalfEdge1;
+			Convex->GetHalfEdges(CurrEdgeIdx, HalfEdge0, HalfEdge1);
+
+			HalfEdgeToEdge[HalfEdge0] = CurrEdgeIdx;
+			HalfEdgeToEdge[HalfEdge1] = CurrEdgeIdx;
+		}
+	}
+
+	}
\ No newline at end of file
diff --git a/Engine/Plugins/Experimental/Buoyancy/Source/Runtime/Private/BuoyancySubsystem.cpp b/Engine/Plugins/Experimental/Buoyancy/Source/Runtime/Private/BuoyancySubsystem.cpp
--- a/Engine/Plugins/Experimental/Buoyancy/Source/Runtime/Private/BuoyancySubsystem.cpp
+++ b/Engine/Plugins/Experimental/Buoyancy/Source/Runtime/Private/BuoyancySubsystem.cpp
@@ -589,7 +589,10 @@ void FBuoyancySubsystemSimCallback::OnPreSimulate_Internal()
 			SplineData = *Input->SplineData;
 
 			// New spline data means we need to reset our spline key cache
-			SplineKeyCache.Reset();
+			if (!bUseAccurateIntegrationForSplines)
+			{
+				SplineKeyCache.Reset();
+			}
 
 			// New spline data also means that any internal storage of water particles
 			// might be out of date. Clear all data that potentially contains references
@@ -610,8 +613,11 @@ void FBuoyancySubsystemSimCallback::OnPreSimulate_Internal()
 			BuoyancySettings = MoveTemp(Input->BuoyancySettings);
 
 			// Set key cache properties
-			SplineKeyCache.SetGridSize(BuoyancySettings->SplineKeyCacheGridSize);
-			SplineKeyCache.SetCacheLimit(BuoyancySettings->SplineKeyCacheLimit);
+			if (!bUseAccurateIntegrationForSplines)
+			{
+				SplineKeyCache.SetGridSize(BuoyancySettings->SplineKeyCacheGridSize);
+				SplineKeyCache.SetCacheLimit(BuoyancySettings->SplineKeyCacheLimit);
+			}
 		}
 	}
 
@@ -626,7 +632,7 @@ void FBuoyancySubsystemSimCallback::OnPreSimulate_Internal()
 	if (bBuoyancyDebugDraw)
 	{
 		// If we're using spline key cache grid, debug draw all cached points
-		if (BuoyancySettings->bSplineKeyCacheGrid)
+		if (BuoyancySettings->bSplineKeyCacheGrid && !bUseAccurateIntegrationForSplines)
 		{
 			SplineKeyCache.ForEachSplineKey([this](const FBuoyancyWaterSplineData& WaterSpline, const FVector& LocalPos, float SplineKey)
 			{
@@ -905,7 +911,7 @@ void FBuoyancySubsystemSimCallback::TrackInteraction(
 			WaterSampler = TSharedPtr<FBuoyancyWaterSampler>(new FBuoyancyConstantSplineSampler(WaterSpline, ClosestPoint, ClosestPointDerivative, ClosestSplineKey, WaterN));
 
 #if ENABLE_DRAW_DEBUG
-			if (bBuoyancyDebugDraw)
+			if (bBuoyancyDebugDraw && !bUseAccurateIntegrationForSplines)
 			{
 				// Spline Color
 				const FColor SplineColor = FColor::Cyan;
@@ -990,17 +996,33 @@ void FBuoyancySubsystemSimCallback::ProcessAccurateInteraction(Chaos::FPBDRigids
 {		
 	const Chaos::FReal DeltaSeconds = GetDeltaTime_Internal();	
 
-	float SubmergedVol;
-	Chaos::FVec3 SubmergedCoM;
-	Chaos::FVec3 WorldForce;
-	Chaos::FVec3 WorldTorque;		
-	float TotalParticleVol;
+	float SubmergedVol = 0;
+	Chaos::FVec3 SubmergedCoM = Chaos::FVec3::ZeroVector;
+	Chaos::FVec3 WorldForce = Chaos::FVec3::ZeroVector;
+	Chaos::FVec3 WorldTorque = Chaos::FVec3::ZeroVector;
+	float TotalParticleVol = 0;
 	
-	BuoyancyAlgorithms::ComputeSubmergedVolumeAndForcesForParticle(BuoyancyParticleData,
-		Interaction.RigidParticle, Interaction.WaterParticle,
-		WaterSampler,
-		Evolution, DeltaSeconds, BuoyancySettings->WaterDensity, BuoyancySettings->WaterDrag,
-		TotalParticleVol, SubmergedVol, SubmergedCoM, WorldForce, WorldTorque);
+	if (!Interaction.WaterSampler.IsValid())
+	{
+		UE_LOG(LogBuoyancySubsystem, Warning, TEXT("Skipped invalid buoyancy interaction "));
+		return;
+	}
+	else if (Interaction.WaterSampler->GetSamplerType() == EWaterSamplerType::ConstantSpline)
+	{
+		BuoyancyAlgorithms::ComputeSubmergedVolumeAndForcesForParticle(BuoyancyParticleData,
+			Interaction.RigidParticle, Interaction.WaterParticle,
+			StaticCastSharedPtr<FBuoyancyConstantSplineSampler>(Interaction.WaterSampler),
+			Evolution, DeltaSeconds, BuoyancySettings->WaterDensity, BuoyancySettings->WaterDrag,
+			TotalParticleVol, SubmergedVol, SubmergedCoM, WorldForce, WorldTorque);		
+	}
+	else if (Interaction.WaterSampler->GetSamplerType() == EWaterSamplerType::ShallowWater)
+	{
+		BuoyancyAlgorithms::ComputeSubmergedVolumeAndForcesForParticle(BuoyancyParticleData,
+			Interaction.RigidParticle, Interaction.WaterParticle,
+			StaticCastSharedPtr<FBuoyancyShallowWaterSampler>(Interaction.WaterSampler),
+			Evolution, DeltaSeconds, BuoyancySettings->WaterDensity, BuoyancySettings->WaterDrag,
+			TotalParticleVol, SubmergedVol, SubmergedCoM, WorldForce, WorldTorque);
+	}
 
 	// Apply forces to the particle
 	{
diff --git a/Engine/Plugins/Experimental/Buoyancy/Source/Runtime/Public/BuoyancyAlgorithms.h b/Engine/Plugins/Experimental/Buoyancy/Source/Runtime/Public/BuoyancyAlgorithms.h
--- a/Engine/Plugins/Experimental/Buoyancy/Source/Runtime/Public/BuoyancyAlgorithms.h
+++ b/Engine/Plugins/Experimental/Buoyancy/Source/Runtime/Public/BuoyancyAlgorithms.h
@@ -6,14 +6,19 @@
 #include "Chaos/Core.h"
 #include "Chaos/ParticleHandleFwd.h"
 #include "Chaos/AABB.h"
+#include "Chaos/Convex.h"
 #include "WaterBodyComponent.h"
 
+
 namespace Chaos
 {
 	class FPBDRigidsEvolutionGBF;
 }
 
 class FBuoyancyWaterSampler;
+class FBuoyancyConstantSplineSampler;
+class FBuoyancyShallowWaterSampler;
+
 struct FBuoyancyParticleData;
 
 namespace BuoyancyAlgorithms
@@ -31,16 +36,118 @@ namespace BuoyancyAlgorithms
 		Chaos::FVec3 SubmergedCoM;
 	};
 
+	// #todo(dmp): adjust these
+	struct FBuoyancyShapeTopologyLimits
+	{
+	public:
+		static const int32 MaxVerticesPerShape = 100;		
+		static const int32 MaxEdgesPerShape = 100;
+		static const int32 MaxIntersectionPointsPerShape = MaxEdgesPerShape;
+		static const int32 MaxVerticesPerFace = 20;
+		static const int32 MaxSubmergedFaceVertices = MaxVerticesPerFace + 1;
+	};
+	
+	// abstract class to sample water for a given particle
+	class FBuoyancyShape
+	{
+	public:
+		FBuoyancyShape() {};
+
+		virtual ~FBuoyancyShape() {}
+
+		virtual void Initialize() = 0;	
+
+		virtual int32 NumVertices() const = 0;
+		virtual int32 NumEdges() const = 0;
+		virtual int32 NumFaces() const = 0;
+		virtual int NumFaceVertices(const int Index) const = 0;
+		virtual int NumMaxIntersectionsPoints() const = 0;
+
+		virtual Chaos::FVec3 GetVertex(int32 Index) const = 0;
+		virtual void GetEdgeVertices(int EdgeIndex, int32& OutIndex0, int32& OutIndex1) const = 0;
+		virtual int32 GetFaceVertex(const int FaceIndex, const int FaceVertexIndex) const = 0;
+		virtual int32 GetFaceEdge(const int FaceIndex, const int FaceEdgeIndex) const = 0;
+		
+		virtual Chaos::FVec3 GetCenter() const = 0;
+	};
+
+	class FBuoyancyBoxShape : public FBuoyancyShape
+	{
+	public:
+		FBuoyancyBoxShape(const Chaos::FAABB3& InLocalBox) : LocalBox(InLocalBox) {};
+		
+		void FORCEINLINE Initialize() {}
+
+		int32 FORCEINLINE NumVertices() const override { return 8; }
+		int32 FORCEINLINE NumEdges() const override { return 12; }
+		int32 FORCEINLINE NumFaces() const override { return 6; }
+		int FORCEINLINE NumFaceVertices(const int Index) const override { return 4; }
+		int FORCEINLINE NumMaxIntersectionsPoints() const override { return 6; }
+
+		Chaos::FVec3 FORCEINLINE GetVertex(int32 Index) const override { return LocalBox.GetVertex(Index); }
+
+		void FORCEINLINE GetEdgeVertices(int EdgeIndex, int32& OutIndex0, int32& OutIndex1) const
+		{
+			const Chaos::FAABBEdge CurrEdge = LocalBox.GetEdge(EdgeIndex);
+			OutIndex0 = CurrEdge.VertexIndex0;
+			OutIndex1 = CurrEdge.VertexIndex1;
+		}
+
+		int32 FORCEINLINE GetFaceVertex(const int FaceIndex, const int FaceVertexIndex) const override { return LocalBox.GetFace(FaceIndex).VertexIndex[FaceVertexIndex]; }
+		int32 FORCEINLINE GetFaceEdge(const int FaceIndex, const int FaceEdgeIndex) const override { return LocalBox.GetFace(FaceIndex).EdgeIndex[FaceEdgeIndex]; }		
+		Chaos::FVec3 FORCEINLINE GetCenter() const override	{ return LocalBox.GetCenter(); }
+
+	private:
+		const Chaos::FAABB3& LocalBox;
+	};
+
+	class FBuoyancyConvexShape : public FBuoyancyShape
+	{
+	public:
+		FBuoyancyConvexShape(const Chaos::FConvex *InConvex) : Convex(InConvex) {};
+
+		virtual void Initialize() override;
+
+		int32 FORCEINLINE NumVertices() const override { return Convex->NumVertices(); }
+		int32 FORCEINLINE NumEdges() const override { return Convex->NumEdges(); }
+		int32 FORCEINLINE NumFaces() const override { return Convex->GetFaces().Num(); }
+		int FORCEINLINE NumFaceVertices(const int Index) const override { return Convex->NumPlaneVertices(Index); }
+		
+		// #todo(dmp): whats the best upper bound for max number of intersection points between a convex and plane?
+		int FORCEINLINE NumMaxIntersectionsPoints() const override { return Convex->NumVertices(); }
+		Chaos::FVec3 FORCEINLINE GetVertex(int32 Index) const override { return Convex->GetVertex(Index); }
+
+		void FORCEINLINE GetEdgeVertices(int EdgeIndex, int32& OutIndex0, int32& OutIndex1) const
+		{
+			OutIndex0 = Convex->GetEdgeVertex(EdgeIndex, 0);
+			OutIndex1 = Convex->GetEdgeVertex(EdgeIndex, 1);
+		}
+
+		int32 FORCEINLINE GetFaceVertex(const int FaceIndex, const int FaceVertexIndex) const override { return Convex->GetPlaneVertex(FaceIndex, FaceVertexIndex);	}
+		int32 FORCEINLINE GetFaceEdge(const int FaceIndex, const int FaceEdgeIndex) const override
+		{
+			const int32 HalfEdgeIndex = Convex->GetPlaneHalfEdge(FaceIndex, FaceEdgeIndex);
+			return HalfEdgeToEdge[HalfEdgeIndex];
+		}
+		
+		Chaos::FVec3 FORCEINLINE GetCenter() const override	{ return Convex->GetCenterOfMass();	}
+
+	private:
+		const Chaos::FConvex* Convex;		
+		
+		TArray<int32, TInlineAllocator<FBuoyancyShapeTopologyLimits::MaxEdgesPerShape * 2>> HalfEdgeToEdge;
+	};
+
 	// Compute the effective volume of an entire particle based on its material
 	// density and mass.
 	Chaos::FRealSingle ComputeParticleVolume(const Chaos::FPBDRigidsEvolutionGBF& Evolution, const Chaos::FGeometryParticleHandle* Particle);
 
 	// Compute the effective volume of a shape. This method must reflect the
 	// maximum possible output value of the non-scaled ComputeSubmergedVolume.
-	Chaos::FRealSingle ComputeShapeVolume(const Chaos::FGeometryParticleHandle* Particle);
+	Chaos::FRealSingle ComputeShapeVolume(const Chaos::FGeometryParticleHandle* Particle, const bool UseBoundingBoxVolume = true);
 
 	//
-	void ScaleSubmergedVolume(const Chaos::FPBDRigidsEvolutionGBF& Evolution, const Chaos::FGeometryParticleHandle* Particle, Chaos::FRealSingle& SubmergedVol, Chaos::FRealSingle& TotalVol);
+	void ScaleSubmergedVolume(const Chaos::FPBDRigidsEvolutionGBF& Evolution, const Chaos::FGeometryParticleHandle* Particle, const bool UseBoundingBoxVolume, Chaos::FRealSingle& SubmergedVol, Chaos::FRealSingle& TotalVol);
 
 	// Compute an approximate volume and center of mass of particle B submerged in particle A,
 	// adjusting for the volume of the object based on the material density and mass of the object
@@ -70,30 +177,37 @@ namespace BuoyancyAlgorithms
 		const Chaos::FVec3& WaterVel, const Chaos::FVec3& WaterN, Chaos::FVec3& OutDeltaV, Chaos::FVec3& OutDeltaW);
 
 	// given a particle, loop over the contained shapes and accumulate force/torque/submerged CoM values
+	template <typename SamplerType>
 	void ComputeSubmergedVolumeAndForcesForParticle(FBuoyancyParticleData& ParticleData,
 		const Chaos::FGeometryParticleHandle* SubmergedParticle, const Chaos::FGeometryParticleHandle* WaterParticle,
-		TSharedPtr<FBuoyancyWaterSampler> WaterSampler,
+		TSharedPtr<SamplerType> WaterSampler,
 		const Chaos::FPBDRigidsEvolution& Evolution, const float DeltaSeconds, const float WaterDensity, const float WaterDrag,
 		float& OutTotalParticleVol, float& OutTotalSubmergedVol, Chaos::FVec3& OutTotalSubmergedCoM, Chaos::FVec3& OutTotalForce, Chaos::FVec3& OutTotalTorque);
 
 	// given a shape, compute the submerged volume and accumulate forces
 	// this is done in a single function call because of the iterative nature of the algorithm
+	template <typename ShapeType, typename SamplerType>
 	void ComputeSubmergedVolumeAndForcesForShape(
-		const Chaos::FGeometryParticleHandle* SubmergedParticle,
+		const Chaos::FGeometryParticleHandle* SubmergedParticle, const ShapeType& BoxShape,
 		const Chaos::FPBDRigidsEvolution& Evolution, float DeltaSeconds, const float WaterDensity, const float WaterDrag,
-		const Chaos::FImplicitObject* Implicit, const Chaos::FRigidTransform3& RelativeTransform, const int32 RootObjectIndex, const int32 ObjectIndex, const int32 LeafObjectIndex,
-		TSharedPtr<FBuoyancyWaterSampler> WaterSampler,
-		Chaos::FVec3& OutWaterP, Chaos::FVec3& OutWaterN, 
+		const Chaos::FRigidTransform3 &ShapeWorldTransform, TSharedPtr<SamplerType> WaterSampler,
 		float& OutSubmergedVol, Chaos::FVec3& OutSubmergedCoM,
-		Chaos::FVec3& OutForce, Chaos::FVec3& OutTorque);
+		Chaos::FVec3& OutForce, Chaos::FVec3& OutTorque,
+		Chaos::FVec3& OutBouyancyForce, Chaos::FVec3& OutBuoyancyTorque,
+		float& OutForceWeight);
 
 	// find intersection points between a plane and aabbox
-	void FORCEINLINE FindAllIntersectionPoints(const Chaos::FVec3& WaterP, const Chaos::FVec3& WaterN, const Chaos::FAABB3& LocalBox, const FVector WorldVertexPosition[8],
-		FVector EdgeToIntersectionPoint[12], int32& NumIntersections, TArray<FVector, TFixedAllocator<6>>& OutOrderedIntersectionPoints, FVector& OutIntersectionCenter);
+	template <typename ShapeType>
+	void FORCEINLINE FindAllIntersectionPoints(const Chaos::FVec3& WaterP, const Chaos::FVec3& WaterN, const ShapeType& BoxShape,
+		const TArray<FVector, TInlineAllocator<FBuoyancyShapeTopologyLimits::MaxVerticesPerShape>> &WorldVertexPosition,
+		TMap<int32, FVector>& EdgeToIntersectionPoint, int32& NumIntersections,
+		TArray<FVector, TInlineAllocator<FBuoyancyShapeTopologyLimits::MaxIntersectionPointsPerShape>>& OutOrderedIntersectionPoints, FVector& OutIntersectionCenter);
 
 	// sort intersection points by angle
-	void FORCEINLINE SortIntersectionPointsByAngle(const Chaos::FVec3& WaterP, const Chaos::FVec3& WaterN, const Chaos::FVec3& IntersectionCenter,
-		const FVector EdgeToIntersectionPoint[12], const int32 NumIntersections, TArray<FVector, TFixedAllocator<6>>& OutOrderedIntersectionPoints);
+	template <typename ShapeType>
+	void FORCEINLINE SortIntersectionPointsByAngle(const Chaos::FVec3& WaterP, const Chaos::FVec3& WaterN, const Chaos::FVec3& IntersectionCenter, const ShapeType& BoxShape,
+		const TMap<int, FVector>& EdgeToIntersectionPoint,
+		TArray<FVector, TInlineAllocator<FBuoyancyShapeTopologyLimits::MaxIntersectionPointsPerShape>>& OutOrderedIntersectionPoints);
 
 	bool FORCEINLINE EdgePlaneIntersection(const Chaos::FVec3& WaterP, const Chaos::FVec3& WaterN, const Chaos::FVec3& V0, const Chaos::FVec3& V1, Chaos::FVec3& IntersectionPoint);
 
