new drag model
update buoyancy tests
make sure physics asset for raft computes mass with correct material

diff --git a/Engine/Plugins/Experimental/Buoyancy/Source/Runtime/Private/BuoyancyAlgorithms.cpp b/Engine/Plugins/Experimental/Buoyancy/Source/Runtime/Private/BuoyancyAlgorithms.cpp
--- a/Engine/Plugins/Experimental/Buoyancy/Source/Runtime/Private/BuoyancyAlgorithms.cpp
+++ b/Engine/Plugins/Experimental/Buoyancy/Source/Runtime/Private/BuoyancyAlgorithms.cpp
@@ -649,8 +649,7 @@ namespace BuoyancyAlgorithms
 		OutTotalSubmergedCoM = FVec3::ZeroVector;
 		OutTotalForce = FVec3::ZeroVector;
 		OutTotalTorque = FVec3::ZeroVector;
-		
-		float OutTotalForceWeight = 0;
+
 		Chaos::FVec3 OutTotalBuoyancyForce = FVec3::ZeroVector;
 		Chaos::FVec3 OutTotalBuoyancyTorque = FVec3::ZeroVector;
 
@@ -675,7 +674,7 @@ namespace BuoyancyAlgorithms
 			[&Evolution, DeltaSeconds, WaterDensity, WaterDrag, SubmergedParticle, &WaterSampler, 
 			ParticleIndex, &ShapeInstances, WaterShapeType, WaterShapeInstance, &ParticleWorldTransform,
 			&SubmergedShapes, &OutTotalSubmergedVol, &OutTotalSubmergedCoM,
-			&OutTotalForce, &OutTotalTorque, &OutTotalBuoyancyForce, &OutTotalBuoyancyTorque, &OutTotalForceWeight]
+			&OutTotalForce, &OutTotalTorque, &OutTotalBuoyancyForce, &OutTotalBuoyancyTorque]
 			(const FImplicitObject* Implicit, const FRigidTransform3& RelativeTransform, const int32 RootObjectIndex, const int32 ObjectIndex, const int32 LeafObjectIndex)
 			{								
 				const int32 ShapeIndex = (ShapeInstances.IsValidIndex(RootObjectIndex)) ? RootObjectIndex : 0;
@@ -698,8 +697,7 @@ namespace BuoyancyAlgorithms
 				float OutSubmergedVol;
 				FVec3 OutSubmergedCoM;
 				FVec3 OutForce, OutTorque;
-				FVec3 OutBuoyancyForce, OutBuoyancyTorque;
-				float OutForceWeight;
+				FVec3 OutBuoyancyForce, OutBuoyancyTorque;				
 
 				// Get the world-space bounds of shape A
 				FRigidTransform3 ShapeWorldTransform = RelativeTransform * ParticleWorldTransform;
@@ -754,7 +752,7 @@ namespace BuoyancyAlgorithms
 					ComputeSubmergedVolumeAndForcesForShape<FBuoyancyConvexShape, SamplerType>(SubmergedParticle, ConvexShapeQuery,
 						Evolution, DeltaSeconds, WaterDensity, WaterDrag,
 						ShapeWorldTransformScaled, WaterSampler,
-						OutSubmergedVol, OutSubmergedCoM, OutForce, OutTorque, OutBuoyancyForce, OutBuoyancyTorque, OutForceWeight);
+						OutSubmergedVol, OutSubmergedCoM, OutForce, OutTorque, OutBuoyancyForce, OutBuoyancyTorque);
 				}
 				else
 				{					
@@ -764,7 +762,7 @@ namespace BuoyancyAlgorithms
 					ComputeSubmergedVolumeAndForcesForShape<FBuoyancyBoxShape, SamplerType>(SubmergedParticle, BoxShapeQuery,
 						Evolution, DeltaSeconds, WaterDensity, WaterDrag,
 						ShapeWorldTransform, WaterSampler,
-						OutSubmergedVol, OutSubmergedCoM, OutForce, OutTorque, OutBuoyancyForce, OutBuoyancyTorque, OutForceWeight);
+						OutSubmergedVol, OutSubmergedCoM, OutForce, OutTorque, OutBuoyancyForce, OutBuoyancyTorque);
 				}
 
 #if ENABLE_DRAW_DEBUG
@@ -787,8 +785,6 @@ namespace BuoyancyAlgorithms
 
 					OutTotalSubmergedVol += OutSubmergedVol;
 
-					OutTotalForceWeight += OutForceWeight;
-
 					OutTotalBuoyancyForce += OutBuoyancyForce;
 					OutTotalBuoyancyTorque += OutBuoyancyTorque;
 
@@ -796,9 +792,6 @@ namespace BuoyancyAlgorithms
 				}
 			});
 
-		OutTotalForce /= OutTotalForceWeight;
-		OutTotalTorque /= OutTotalForceWeight;
-
 		OutTotalForce += OutTotalBuoyancyForce;
 		OutTotalTorque += OutTotalBuoyancyTorque;
 
@@ -828,8 +821,7 @@ namespace BuoyancyAlgorithms
 		const Chaos::FRigidTransform3 &ShapeWorldTransform, TSharedPtr<SamplerType> WaterSampler,
 		float& OutSubmergedVol, Chaos::FVec3& OutSubmergedCoM,
 		Chaos::FVec3& OutForce, Chaos::FVec3& OutTorque,
-		Chaos::FVec3& OutBuoyancyForce, Chaos::FVec3& OutBuoyancyTorque,
-		float& OutForceWeight)
+		Chaos::FVec3& OutBuoyancyForce, Chaos::FVec3& OutBuoyancyTorque)
 	{
 		SCOPE_CYCLE_COUNTER(STAT_BuoyancyAlgorithms_ComputeSubmergedVolumeAndForcesForShape)
 
@@ -941,8 +933,6 @@ namespace BuoyancyAlgorithms
 		InteriorRefPoint += IntersectionCenter * NumIntersectionPoints;
 		InteriorRefPointCount += NumIntersectionPoints;
 		InteriorRefPoint /= InteriorRefPointCount;
-
-		OutForceWeight = 0;
 		
 		const FVec3 WorldCoM = SubmergedGeneric->PCom();
 
@@ -1042,15 +1032,12 @@ namespace BuoyancyAlgorithms
 					FVector TotalWorldForce;
 					FVector TotalWorldTorque;
 
-					// #todo(dmp): we scale drag to make it closer to the old algorithm for now
-					const float DragFactor = .05 * WaterDrag;
 					BuoyancyAlgorithms::ComputeFluidForceForTriangle(
-						DragFactor, DeltaSeconds, RigidParticle, WorldCoM, TriBaryCenter, TriNormal, TriArea, TetVolume,
+						WaterDrag, WaterDensity, DeltaSeconds, RigidParticle, WorldCoM, TriBaryCenter, TriNormal, TriArea, TetVolume,
 						CurrWaterVelocity, ShapeWaterP, ShapeWaterN, TotalWorldForce, TotalWorldTorque);
 
-					OutForce += TotalWorldForce * TriArea;
-					OutTorque += TotalWorldTorque * TriArea;
-					OutForceWeight += TriArea;
+					OutForce += TotalWorldForce;
+					OutTorque += TotalWorldTorque;					
 
 #if ENABLE_DRAW_DEBUG
 					if (bBuoyancyDebugDraw)
@@ -1281,7 +1268,7 @@ namespace BuoyancyAlgorithms
 
 	// #todo(dmp): optimization- precalculate particle v, w, m given they never change between calls to compute forces
 	void ComputeFluidForceForTriangle(const float WaterDrag,
-		const float DeltaSeconds,
+		const float DeltaSeconds, const float WaterDensity,
 		const Chaos::FPBDRigidParticleHandle* RigidParticle, const FVector WorldCoM,
 		const FVector &TriBaryCenter, const FVector &TriNormal, const float TriArea, const float TetVolume,
 		const FVector &WaterVelocity, const FVector &WaterP, const FVector &WaterN,
@@ -1306,31 +1293,22 @@ namespace BuoyancyAlgorithms
 		const FVec3 RelativeVelocity = WaterVelocityOnPlane - SubmergedParticleVelocity;
 		const float RelativeVelocityMag = RelativeVelocity.Length();
 
-		if (RelativeVelocityMag < SMALL_NUMBER)
+		// test if this triangle is influenced by the water based on the normal since we have closed shapes only.  This drag algorithm
+		// is derived for two sided planes
+		const float FacingTest = RelativeVelocity.Dot(TriNormal);
+
+		if (RelativeVelocityMag < SMALL_NUMBER || FacingTest < 0.f)
 		{
 			return;
 		}
-		
-		const FVec3 AccelerationToSetVelocity = RelativeVelocity / DeltaSeconds;
-		FVec3 ForceFromWaterVelocity = AccelerationToSetVelocity * RigidParticle->M();
-
-		// #todo(dmp): make friction/drag/other properties more of a material property
-
-		// attenuate based on angle to inward facing tri normal					
-		// we reject angles > 90 since that would mean we are hitting a backface and no force should be applied					
-		// default to 1 if we have a still pool of water, we still want this to have a valid response so we don't have
-		// a discontinuity.
-		float FacingTest = RelativeVelocity.Dot(-TriNormal) / RelativeVelocityMag;
-		FacingTest = FacingTest > 0.f ? FacingTest : 0.f;
-
-		// apply some drag based on facing ratio...could make this more complex
-		ForceFromWaterVelocity *= FacingTest;
 
-		// #todo(dmp): drag should include a term based on relative velocity?
-		// ForceFromWaterVelocity *= RelativeVelocity.Length();
+		// go with the flow combined drag and lift
+		// https://www.yousufsoliman.com/projects/download/going-with-the-flow.pdf
+		//const FVec3 ForceFromWaterVelocity = .5 * WaterDensity * RelativeVelocityMag * RelativeVelocity.Dot(TriNormal) * TriNormal * WaterDrag * TriArea;	
 
-		// drag coefficient			
-		ForceFromWaterVelocity *= WaterDrag;
+		// https://www.cemyuksel.com/research/waveparticles/cem_yuksel_dissertation.pdf
+		const float A = TriArea * RelativeVelocity.Dot(TriNormal) / RelativeVelocityMag;
+		const FVec3 ForceFromWaterVelocity = .5 * WaterDensity * RelativeVelocityMag * RelativeVelocity * A * WaterDrag;	
 
 		// compute torque based on the linear force we apply
 		const FVec3 TorqueFromWaterVelocity = Chaos::FVec3::CrossProduct(WorldCOMToForcePos, ForceFromWaterVelocity);
diff --git a/Engine/Plugins/Experimental/Buoyancy/Source/Runtime/Public/BuoyancyAlgorithms.h b/Engine/Plugins/Experimental/Buoyancy/Source/Runtime/Public/BuoyancyAlgorithms.h
--- a/Engine/Plugins/Experimental/Buoyancy/Source/Runtime/Public/BuoyancyAlgorithms.h
+++ b/Engine/Plugins/Experimental/Buoyancy/Source/Runtime/Public/BuoyancyAlgorithms.h
@@ -193,8 +193,7 @@ namespace BuoyancyAlgorithms
 		const Chaos::FRigidTransform3 &ShapeWorldTransform, TSharedPtr<SamplerType> WaterSampler,
 		float& OutSubmergedVol, Chaos::FVec3& OutSubmergedCoM,
 		Chaos::FVec3& OutForce, Chaos::FVec3& OutTorque,
-		Chaos::FVec3& OutBouyancyForce, Chaos::FVec3& OutBuoyancyTorque,
-		float& OutForceWeight);
+		Chaos::FVec3& OutBouyancyForce, Chaos::FVec3& OutBuoyancyTorque);
 
 	// find intersection points between a plane and aabbox
 	template <typename ShapeType>
@@ -217,7 +216,7 @@ namespace BuoyancyAlgorithms
 
 	// compute the force the fluid exerts on a triangle
 	void FORCEINLINE ComputeFluidForceForTriangle(const float WaterDrag,
-		const float DeltaSeconds,
+		const float DeltaSeconds, const float WaterDensity,
 		const Chaos::FPBDRigidParticleHandle* RigidParticle, const FVector WorldCoM,
 		const FVector &TriBaryCenter, const FVector &TriNormal, const float TriArea, const float TetVolume,
 		const FVector &WaterVelocity, const FVector &WaterP, const FVector &WaterN,
