default drag for accurate integration

diff --git a/Engine/Plugins/Experimental/Buoyancy/Source/Runtime/Private/BuoyancySubsystem.cpp b/Engine/Plugins/Experimental/Buoyancy/Source/Runtime/Private/BuoyancySubsystem.cpp
--- a/Engine/Plugins/Experimental/Buoyancy/Source/Runtime/Private/BuoyancySubsystem.cpp
+++ b/Engine/Plugins/Experimental/Buoyancy/Source/Runtime/Private/BuoyancySubsystem.cpp
@@ -57,7 +57,7 @@ int32 bUseShallowWaterSimulation = 1;
 static FAutoConsoleVariableRef CVarUseShallowWaterSimulation(
 	TEXT("p.Buoyancy.bUseShallowWaterSimulation"),
 	bUseShallowWaterSimulation,
-	TEXT("new buoancy method"),
+	TEXT("Accurate buoyancy method"),
 	ECVF_Scalability
 );
 
@@ -65,7 +65,15 @@ int32 bUseAccurateIntegrationForSplines = 0;
 static FAutoConsoleVariableRef CVarUseAccurateIntegrationForSplines(
 	TEXT("p.Buoyancy.bUseAccurateIntegrationForSplines"),
 	bUseAccurateIntegrationForSplines,
-	TEXT("new buoancy method for splines"),
+	TEXT("Accurate buoyancy method for splines"),
+	ECVF_Scalability
+);
+
+float AccurateIntegrationDragMultiplier = 0.0075;
+static FAutoConsoleVariableRef CVarAccurateIntegrationDragMultiplier(
+	TEXT("p.Buoyancy.AccurateIntegrationDragMultiplier"),
+	AccurateIntegrationDragMultiplier,
+	TEXT("Accurate buoyancy method for splines"),
 	ECVF_Scalability
 );
 
@@ -1002,6 +1010,8 @@ void FBuoyancySubsystemSimCallback::ProcessAccurateInteraction(Chaos::FPBDRigids
 	Chaos::FVec3 WorldTorque = Chaos::FVec3::ZeroVector;
 	float TotalParticleVol = 0;
 	
+	const float AccurateIntegrationDrag = AccurateIntegrationDragMultiplier * BuoyancySettings->WaterDrag;
+
 	if (!Interaction.WaterSampler.IsValid())
 	{
 		UE_LOG(LogBuoyancySubsystem, Warning, TEXT("Skipped invalid buoyancy interaction "));
@@ -1012,7 +1022,7 @@ void FBuoyancySubsystemSimCallback::ProcessAccurateInteraction(Chaos::FPBDRigids
 		BuoyancyAlgorithms::ComputeSubmergedVolumeAndForcesForParticle(BuoyancyParticleData,
 			Interaction.RigidParticle, Interaction.WaterParticle,
 			StaticCastSharedPtr<FBuoyancyConstantSplineSampler>(Interaction.WaterSampler),
-			Evolution, DeltaSeconds, BuoyancySettings->WaterDensity, BuoyancySettings->WaterDrag,
+			Evolution, DeltaSeconds, BuoyancySettings->WaterDensity, AccurateIntegrationDrag,
 			TotalParticleVol, SubmergedVol, SubmergedCoM, WorldForce, WorldTorque);		
 	}
 	else if (Interaction.WaterSampler->GetSamplerType() == EWaterSamplerType::ShallowWater)
@@ -1020,7 +1030,7 @@ void FBuoyancySubsystemSimCallback::ProcessAccurateInteraction(Chaos::FPBDRigids
 		BuoyancyAlgorithms::ComputeSubmergedVolumeAndForcesForParticle(BuoyancyParticleData,
 			Interaction.RigidParticle, Interaction.WaterParticle,
 			StaticCastSharedPtr<FBuoyancyShallowWaterSampler>(Interaction.WaterSampler),
-			Evolution, DeltaSeconds, BuoyancySettings->WaterDensity, BuoyancySettings->WaterDrag,
+			Evolution, DeltaSeconds, BuoyancySettings->WaterDensity, AccurateIntegrationDrag,
 			TotalParticleVol, SubmergedVol, SubmergedCoM, WorldForce, WorldTorque);
 	}
 
