new buoyancy queries to enable cloth/water interaction

diff --git a/Engine/Plugins/Experimental/Buoyancy/Source/Runtime/Private/BuoyancySubsystem.cpp b/Engine/Plugins/Experimental/Buoyancy/Source/Runtime/Private/BuoyancySubsystem.cpp
--- a/Engine/Plugins/Experimental/Buoyancy/Source/Runtime/Private/BuoyancySubsystem.cpp
+++ b/Engine/Plugins/Experimental/Buoyancy/Source/Runtime/Private/BuoyancySubsystem.cpp
@@ -34,7 +34,7 @@
 #include "Chaos/Collision/CollisionFilter.h"
 #include "Chaos/Collision/CollisionUtil.h"
 #include "Chaos/Sphere.h"
-
+#include "Engine/OverlapResult.h"
 //
 // CVars
 //
@@ -344,6 +344,135 @@ void UBuoyancySubsystem::UpdateNetMode()
 	}
 }
 
+// #todo(dmp): the current use of this method is on whatever thread Cloth is evaluated with, which is not the physics thread,
+//  but we are using the FBuoyancyWaterSplineData, which is stored on the PT after it is created from the Buoyancy plugin
+bool UBuoyancySubsystem::QueryWaterBody(const FVector& InputPosition, const TSharedPtr<FBuoyancyWaterSplineData> WaterData,
+	FVector& WaterVel, FVector& WaterPlaneN, FVector& WaterPlanePos)
+{		
+	float ClosestSplineKey = 0;
+	WaterPlanePos = FVector::ZeroVector;
+	WaterPlaneN = FVector::ZeroVector;
+	WaterVel = FVector::ZeroVector;
+
+	if (SplineData == nullptr || WaterData == nullptr)
+	{
+		return false;
+	}
+
+	if (WaterData->ShouldSampleFromShallowWaterSimulation())
+	{
+		// eval shallow water - closest point is just the current point adjusted upward in Z		
+		float WaterHeight;
+		float WaterDepth;
+		WaterData->ShallowWaterSimData->SampleShallowWaterSimulationAtPosition(InputPosition, WaterVel, WaterHeight, WaterDepth);
+		
+		// #todo(dmp): if there is no water here, then return.  This is an approximation that there is water underneath this rigid that might cause interaction
+		// ideally, we'd multisample this since we can lose some water interactions this way
+		if (WaterDepth > 1e-5)
+		{			
+			WaterPlanePos = InputPosition;
+			WaterPlanePos.Z = WaterHeight;
+		
+			// compute water plane normal
+			WaterPlaneN = WaterData->ShallowWaterSimData->ComputeShallowWaterSimulationNormalAtPosition(InputPosition);
+
+			return true;
+		}
+	}
+	else
+	{
+		// Find water surface at the nearest point on the spline
+		const FVector ParticlePos = InputPosition;
+
+		FVector ClosestPointDerivative = FVector::ZeroVector;
+		bool IsInsideSpline = SimCallback->QuerySpline(ParticlePos, *WaterData, ClosestSplineKey, WaterPlanePos, ClosestPointDerivative, WaterPlaneN);
+
+		if (IsInsideSpline)
+		{
+			WaterVel = WaterData->Velocity->Eval(ClosestSplineKey) * ClosestPointDerivative.GetSafeNormal();
+
+			return true;
+		}		
+	}
+
+	return false;
+}
+
+bool UBuoyancySubsystem::FindOverlappingWaterBodies(const FBox BoundingBox, TArray<const TSharedPtr<FBuoyancyWaterSplineData>>& WaterBodyPhysicsProxies)
+{
+	if (!SplineData)
+	{
+		return false;
+	}
+
+	FCollisionShape BoxShape = FCollisionShape::MakeBox(BoundingBox.GetExtent());
+
+	FCollisionQueryParams QueryParams;
+	QueryParams.bTraceComplex = false;
+	QueryParams.bReturnPhysicalMaterial = false;			
+
+	FCollisionObjectQueryParams ObjectQueryParams;
+	ObjectQueryParams.AddObjectTypesToQuery(ECC_WorldStatic);
+
+	TArray<FOverlapResult> Overlaps;
+			
+	// perform overlap test
+	bool bHit = GetWorld()->OverlapMultiByObjectType(
+		Overlaps,
+		BoundingBox.GetCenter(),
+		FQuat::Identity,
+		ObjectQueryParams,
+		BoxShape,
+		QueryParams
+	);
+
+	WaterBodyPhysicsProxies.Empty();
+				
+	bool bFoundWaterBody = false;
+
+	// if we have a hit, find the water bodies we are colliding with
+	if (bHit)
+	{
+		for (const FOverlapResult& Result : Overlaps)
+		{
+			UWaterBodyComponent* WaterBodyComponent = Cast<UWaterBodyComponent>(Result.Component->GetAttachParent());
+			if (WaterBodyComponent)
+			{						
+				for (UPrimitiveComponent* WaterPrimitiveComponent : WaterBodyComponent->GetCollisionComponents(true))
+				{
+					// Add each object (probably just one) to the objects list
+					if (WaterPrimitiveComponent)
+					{
+						FBodyInstance& BodyInstance = WaterPrimitiveComponent->BodyInstance;
+						if (BodyInstance.IsValidBodyInstance())
+						{
+							Chaos::FSingleParticlePhysicsProxy* WaterProxy = BodyInstance.GetPhysicsActor();
+							
+							// if we found a valid physics actor, get the PT spline data associated with the particle
+							// #todo(dmp): resolve issues with threading since this method is called from the GT but we are
+							// querying PT data here that is created from the Buoyancy plugin.
+							if (WaterProxy)
+							{
+								bFoundWaterBody = true;
+								Chaos::FRigidBodyHandle_External& Body_External = WaterProxy->GetGameThreadAPI();								
+								
+								const TSharedPtr<FBuoyancyWaterSplineData> *TmpData = SplineData->GetData_PT(WaterProxy->GetGameThreadAPI());								
+								if (TmpData)
+								{
+									WaterBodyPhysicsProxies.Add(*TmpData);
+								}
+							}
+						}
+					}
+				}
+			}
+		}
+		return bFoundWaterBody;
+	}
+
+	return false;
+}
+
 void UBuoyancySubsystem::UpdateSplineData()
 {
 	SCOPE_CYCLE_COUNTER(STAT_BuoyancySubsystem_UpdateWaterBodiesList)
@@ -788,6 +917,135 @@ void FBuoyancySubsystemSimCallback::TrackInteractions(
 	});
 }
 
+bool FBuoyancySubsystemSimCallback::QuerySpline(
+	const FVector &QueryPos, 
+	const FBuoyancyWaterSplineData& WaterSpline,
+	float &ClosestSplineKey,
+	FVector &ClosestPoint,
+	FVector &ClosestPointDerivative,
+	FVector &WaterN)
+{
+	SCOPE_CYCLE_COUNTER(STAT_BuoyancySubsystem_SplineEvaluation)
+
+	ClosestSplineKey = 0;
+	ClosestPoint = FVector::ZeroVector;
+	WaterN = FVector::ZeroVector;	
+
+	// Find water surface at the nearest point on the spline	
+	{
+		SCOPE_CYCLE_COUNTER(STAT_BuoyancySubsystem_SplineEvaluation_FindNearest)
+
+		if (BuoyancySettings->bSplineKeyCacheGrid)
+		{
+			ClosestSplineKey = SplineKeyCache.GetClosestSplineKey(WaterSpline, QueryPos);
+		}
+		else
+		{
+			const FVector ParticleLocalPos = WaterSpline.Transform.InverseTransformPosition(QueryPos);
+			float ParticleDistance;
+			ClosestSplineKey = WaterSpline.Position.FindNearest(ParticleLocalPos, ParticleDistance);
+		}
+	}
+	{
+		SCOPE_CYCLE_COUNTER(STAT_BuoyancySubsystem_SplineEvaluation_Eval)
+		ClosestPoint = WaterSpline.Transform.TransformPosition(WaterSpline.Position.Eval(ClosestSplineKey));
+	}
+
+	// Get water surface level and normal
+	// Don't build interaction if outside of water body
+	{			
+		// Find water surface at the nearest point on the spline			
+		ClosestPointDerivative = WaterSpline.Transform.TransformVector(
+		WaterSpline.Position.EvalDerivative(ClosestSplineKey));
+
+		// Water normal direction depends on body type
+		if (WaterSpline.BodyType == EWaterBodyType::River)
+		{
+			// River water normal can be determined by the relationship of the 
+			// derivative of the spline position to the up vector.
+			//
+			// NOTE: This calculation breaks down in the limit of purely
+			// vertical water
+			const FVector SplineRight = FVector::CrossProduct(FVector::UpVector, ClosestPointDerivative);
+			const FVector SplineUp = FVector::CrossProduct(ClosestPointDerivative, SplineRight);
+			WaterN = SplineUp.GetSafeNormal();
+		}
+		else
+		{
+			WaterN = FVector::UpVector;
+		}
+
+		// Project the position difference onto the water surface
+		const FVector Diff = ClosestPoint - QueryPos;
+		const FVector LateralDiff = Diff - (WaterN * FVector::DotProduct(WaterN, Diff));
+
+		// Different water body types have different ways of determining
+		// whether a point is laterally inside their volume.
+		switch (WaterSpline.BodyType)
+		{
+			case EWaterBodyType::River:
+			{
+				if (WaterSpline.Width.IsSet())
+				{
+					// If distance to spline is greater than the width of the spline,
+					// then this is a river and we're outside of it.
+					const float Width = WaterSpline.Width->Eval(ClosestSplineKey);
+					const float DistSq = FVector::DotProduct(LateralDiff, LateralDiff);
+					const float WidthSq = Width * Width * .25f;
+					if (DistSq > WidthSq) { return false; }
+				}
+				break;
+			}
+
+			case EWaterBodyType::Lake:
+			{
+				// Determine if we're inside the lake by projecting the horizontal spline
+				// diff onto the cross product of the spline direction and the up-vector
+				// (ie, the right-vector)
+				const FVector RightVector = FVector::CrossProduct(ClosestPointDerivative, FVector::UpVector);
+				const float DiffProj = FVector::DotProduct(RightVector, LateralDiff);
+				if (DiffProj < SMALL_NUMBER) { return false; }
+				break;
+			}
+		}
+			
+#if ENABLE_DRAW_DEBUG
+		if (bBuoyancyDebugDraw && !bUseAccurateIntegrationForSplines)
+		{
+			// Spline Color
+			const FColor SplineColor = FColor::Cyan;
+
+			// Draw projection onto the line
+			const FVector SurfacePoint = ClosestPoint - LateralDiff;
+			Chaos::FDebugDrawQueue::GetInstance().DrawDebugLine(QueryPos, SurfacePoint, SplineColor, false, -1.f, -1, 6.f);
+			Chaos::FDebugDrawQueue::GetInstance().DrawDebugLine(SurfacePoint, ClosestPoint, SplineColor, false, -1.f, -1, 3.f);
+
+			// Draw a section of the spline near the spline key
+			Chaos::FVec3 PrevPoint;
+			bool bFirst = true;
+			for (float SplineKey = ClosestSplineKey - .1f; SplineKey <= ClosestSplineKey + .1f; SplineKey += .05f)
+			{
+				const FVector SplinePoint = WaterSpline.Transform.TransformPosition(WaterSpline.Position.Eval(SplineKey));
+				if (bFirst)
+				{
+					bFirst = false;
+				}
+				else
+				{
+					Chaos::FDebugDrawQueue::GetInstance().DrawDebugDirectionalArrow(PrevPoint, SplinePoint, 15.f, SplineColor, false, -1.f, -1, 3.f);
+				}
+				PrevPoint = SplinePoint;
+			}
+
+			// Draw water surface normal at the surface point
+			Chaos::FDebugDrawQueue::GetInstance().DrawDebugDirectionalArrow(SurfacePoint, SurfacePoint + (WaterN * 100.f), 20.f, FColor::Green, false, -1.f, -1, 3.f);
+		}
+#endif
+	}
+
+	return true;
+}
+
 void FBuoyancySubsystemSimCallback::TrackInteraction(
 	Chaos::FPBDRigidsEvolution& Evolution,
 	Chaos::FGeometryParticleHandle* WaterParticle,
@@ -836,125 +1094,18 @@ void FBuoyancySubsystemSimCallback::TrackInteraction(
 
 		float ClosestSplineKey = 0;
 		FVector ClosestPoint = FVector::ZeroVector;
+		FVector ClosestPointDerivative = FVector::ZeroVector;
 		FVector WaterN = FVector::ZeroVector;
-		FVector WaterVel = FVector::ZeroVector;
-
+		
 		// Find water surface at the nearest point on the spline
 		const FVector ParticlePos = RigidParticle->XCom();
+		bool IsInslideSpline = QuerySpline(ParticlePos, WaterSpline, ClosestSplineKey, ClosestPoint, ClosestPointDerivative, WaterN);
+		if (!IsInslideSpline)
 		{
-			SCOPE_CYCLE_COUNTER(STAT_BuoyancySubsystem_SplineEvaluation_FindNearest)
-
-			if (BuoyancySettings->bSplineKeyCacheGrid)
-			{
-				ClosestSplineKey = SplineKeyCache.GetClosestSplineKey(WaterSpline, ParticlePos);
-			}
-			else
-			{
-				const FVector ParticleLocalPos = WaterSpline.Transform.InverseTransformPosition(ParticlePos);
-				float ParticleDistance;
-				ClosestSplineKey = WaterSpline.Position.FindNearest(ParticleLocalPos, ParticleDistance);
-			}
-		}
-		{
-			SCOPE_CYCLE_COUNTER(STAT_BuoyancySubsystem_SplineEvaluation_Eval)
-			ClosestPoint = WaterSpline.Transform.TransformPosition(WaterSpline.Position.Eval(ClosestSplineKey));
-		}
-
-		// Get water surface level and normal
-		// Don't build interaction if outside of water body
-		{			
-			// Find water surface at the nearest point on the spline			
-			FVector ClosestPointDerivative = WaterSpline.Transform.TransformVector(
-			WaterSpline.Position.EvalDerivative(ClosestSplineKey));
-
-			// Water normal direction depends on body type
-			if (WaterSpline.BodyType == EWaterBodyType::River)
-			{
-				// River water normal can be determined by the relationship of the 
-				// derivative of the spline position to the up vector.
-				//
-				// NOTE: This calculation breaks down in the limit of purely
-				// vertical water
-				const FVector SplineRight = FVector::CrossProduct(FVector::UpVector, ClosestPointDerivative);
-				const FVector SplineUp = FVector::CrossProduct(ClosestPointDerivative, SplineRight);
-				WaterN = SplineUp.GetSafeNormal();
-			}
-			else
-			{
-				WaterN = FVector::UpVector;
-			}
-
-			// Project the position difference onto the water surface
-			const FVector Diff = ClosestPoint - ParticlePos;
-			const FVector LateralDiff = Diff - (WaterN * FVector::DotProduct(WaterN, Diff));
-
-			// Different water body types have different ways of determining
-			// whether a point is laterally inside their volume.
-			switch (WaterSpline.BodyType)
-			{
-				case EWaterBodyType::River:
-				{
-					if (WaterSpline.Width.IsSet())
-					{
-						// If distance to spline is greater than the width of the spline,
-						// then this is a river and we're outside of it.
-						const float Width = WaterSpline.Width->Eval(ClosestSplineKey);
-						const float DistSq = FVector::DotProduct(LateralDiff, LateralDiff);
-						const float WidthSq = Width * Width * .25f;
-						if (DistSq > WidthSq) { return; }
-					}
-					break;
-				}
-
-				case EWaterBodyType::Lake:
-				{
-					// Determine if we're inside the lake by projecting the horizontal spline
-					// diff onto the cross product of the spline direction and the up-vector
-					// (ie, the right-vector)
-					const FVector RightVector = FVector::CrossProduct(ClosestPointDerivative, FVector::UpVector);
-					const float DiffProj = FVector::DotProduct(RightVector, LateralDiff);
-					if (DiffProj < SMALL_NUMBER) { return; }
-					break;
-				}
-			}
-
-			WaterSampler = TSharedPtr<FBuoyancyWaterSampler>(new FBuoyancyConstantSplineSampler(WaterSpline, ClosestPoint, ClosestPointDerivative, ClosestSplineKey, WaterN));
-
-#if ENABLE_DRAW_DEBUG
-			if (bBuoyancyDebugDraw && !bUseAccurateIntegrationForSplines)
-			{
-				// Spline Color
-				const FColor SplineColor = FColor::Cyan;
-
-				// Draw projection onto the line
-				const FVector SurfacePoint = ClosestPoint - LateralDiff;
-				Chaos::FDebugDrawQueue::GetInstance().DrawDebugLine(ParticlePos, SurfacePoint, SplineColor, false, -1.f, -1, 6.f);
-				Chaos::FDebugDrawQueue::GetInstance().DrawDebugLine(SurfacePoint, ClosestPoint, SplineColor, false, -1.f, -1, 3.f);
-
-				// Draw a section of the spline near the spline key
-				Chaos::FVec3 PrevPoint;
-				bool bFirst = true;
-				for (float SplineKey = ClosestSplineKey - .1f; SplineKey <= ClosestSplineKey + .1f; SplineKey += .05f)
-				{
-					const FVector SplinePoint = WaterSpline.Transform.TransformPosition(WaterSpline.Position.Eval(SplineKey));
-					if (bFirst)
-					{
-						bFirst = false;
-					}
-					else
-					{
-						Chaos::FDebugDrawQueue::GetInstance().DrawDebugDirectionalArrow(PrevPoint, SplinePoint, 15.f, SplineColor, false, -1.f, -1, 3.f);
-					}
-					PrevPoint = SplinePoint;
-				}
-
-				// Draw water surface normal at the surface point
-				Chaos::FDebugDrawQueue::GetInstance().DrawDebugDirectionalArrow(SurfacePoint, SurfacePoint + (WaterN * 100.f), 20.f, FColor::Green, false, -1.f, -1, 3.f);
-			}
-#endif
-		}
-	
+			return;
+		}		
 		
+		WaterSampler = TSharedPtr<FBuoyancyWaterSampler>(new FBuoyancyConstantSplineSampler(WaterSpline, ClosestPoint, ClosestPointDerivative, ClosestSplineKey, WaterN));				
 	}
 
 	{
diff --git a/Engine/Plugins/Experimental/Buoyancy/Source/Runtime/Public/BuoyancySubsystem.h b/Engine/Plugins/Experimental/Buoyancy/Source/Runtime/Public/BuoyancySubsystem.h
--- a/Engine/Plugins/Experimental/Buoyancy/Source/Runtime/Public/BuoyancySubsystem.h
+++ b/Engine/Plugins/Experimental/Buoyancy/Source/Runtime/Public/BuoyancySubsystem.h
@@ -75,6 +75,22 @@ struct FBuoyancySettings
 	uint32 SplineKeyCacheLimit = 256;
 };
 
+class FBuoyancyCollisionData
+{
+	public:
+		FBuoyancyCollisionData(UBuoyancySubsystem *InTheBuoyancySubsystem, TArray<const TSharedPtr<FBuoyancyWaterSplineData>> InWaterBodyCollisionData)
+		: TheBuoyancySubsystem(InTheBuoyancySubsystem)
+		, WaterBodyCollisionData(InWaterBodyCollisionData) 
+		{}
+		
+		FBuoyancyCollisionData() 
+		: TheBuoyancySubsystem(nullptr)
+		, WaterBodyCollisionData(TArray<const TSharedPtr<FBuoyancyWaterSplineData>>()) 
+		{}
+
+		UBuoyancySubsystem *TheBuoyancySubsystem;
+		TArray<const TSharedPtr<FBuoyancyWaterSplineData>> WaterBodyCollisionData;
+};
 
 //
 // Buoyancy Subsystem
@@ -111,6 +127,21 @@ public:
 	UE_API SIZE_T GetAllocatedSize() const;
 #endif
 
+	// given a bounding box, perform an overlap test and filter for water bodies.  There could be some potential threading
+	// issues here since we are querying the FBuoyancyWaterSplineDataManager physics thread data presumably from the gt to
+	// do the overlap test.  We need a better way of communicating the buoyancy data so we can capture it on the gt, 
+	// then use it on the physics or other thread (ie: cloth).
+	UE_API bool FindOverlappingWaterBodies(
+		const FBox BoundingBox, TArray<const TSharedPtr<FBuoyancyWaterSplineData>> &WaterBodyPhysicsProxies);
+
+	// query a water body given a FBuoyancyWaterSplineData object representing a water body
+	// note this is designed to only work on the physics thread, but the FBuoyancyWaterSplineData data is 
+	// generated once and then only written to if water bodies change (uncommon).  In the future, we should consider
+	// restructuring usage such that there aren't potential threading issues, but it is problematic since Chaos and
+	// Cloth are evaluated on separate threads.
+	UE_API bool QueryWaterBody(const FVector& InputPosition,  const TSharedPtr<FBuoyancyWaterSplineData> WaterData,
+		FVector& WaterVel, FVector& WaterPlaneN, FVector& WaterPlanePos);
+
 protected:
 
 	// UTickableWorldSubsystem begin interface
@@ -222,13 +253,17 @@ class FBuoyancySubsystemSimCallback : public Chaos::TSimCallbackObject<
 public:
 
 	SIZE_T GetAllocatedSize() const;
+	FSplineKeyCacheGrid &GetSplineKeyCache() { return SplineKeyCache; }
 
+	bool QuerySpline(const FVector &QueryPos, const FBuoyancyWaterSplineData& WaterSpline, float &ClosestSplineKey,
+		FVector &ClosestPoint, FVector &ClosestPointDerivative, FVector &WaterN);
 private:
 
 	virtual void OnPreSimulate_Internal() override;
 	virtual void OnMidPhaseModification_Internal(Chaos::FMidPhaseModifierAccessor& Modifier) override;
 
 	void TrackInteractions(Chaos::FPBDRigidsSolver& PBDSolver, Chaos::FPBDRigidsEvolution& Evolution, Chaos::FMidPhaseModifierAccessor& MidPhaseAccessor);
+	
 	void TrackInteraction(Chaos::FPBDRigidsEvolution& Evolution, Chaos::FGeometryParticleHandle* WaterParticle, Chaos::FPBDRigidParticleHandle* RigidParticle, const FBuoyancyWaterSplineData& WaterSpline, Chaos::FMidPhaseModifier& MidPhase);
 	void ProcessInteractions(Chaos::FPBDRigidsEvolution& Evolution);
 	void ProcessAccurateInteraction(Chaos::FPBDRigidsEvolution& Evolution, FBuoyancyInteraction& Interaction, TSharedPtr<FBuoyancyWaterSampler> WaterSampler);
