ChaosCloth Water Body Interaction
- Subclass of VelocityField that supports Water Body interactions
- New parameters exposed in DF for controlling buoyancy and underwater drag
- Existing drag applies for above water parts of cloth
- Hooks for replacinig existing VelocityField and configuring BuoyancyField in ChaosCloth

Future Work (Depending on various factors)
- ISPC implementation
- Per Triangle water body queries for accurate collisions (needs threading issuesin Buoyancy addressed)

#Synced-CL 49393612

diff --git a/Engine/Plugins/ChaosCloth/ChaosCloth.uplugin b/Engine/Plugins/ChaosCloth/ChaosCloth.uplugin
--- a/Engine/Plugins/ChaosCloth/ChaosCloth.uplugin
+++ b/Engine/Plugins/ChaosCloth/ChaosCloth.uplugin
@@ -29,6 +29,10 @@
 		{
 			"Name": "ChaosCaching",
 			"Enabled": true
+		},
+		{
+			"Name": "Buoyancy",
+			"Enabled": true
 		}
 	]
 }
diff --git a/Engine/Plugins/ChaosCloth/Source/ChaosCloth/ChaosCloth.Build.cs b/Engine/Plugins/ChaosCloth/Source/ChaosCloth/ChaosCloth.Build.cs
--- a/Engine/Plugins/ChaosCloth/Source/ChaosCloth/ChaosCloth.Build.cs
+++ b/Engine/Plugins/ChaosCloth/Source/ChaosCloth/ChaosCloth.Build.cs
@@ -14,7 +14,9 @@ public class ChaosCloth : ModuleRules
                 "ClothingSystemRuntimeInterface",
                 "Engine",
                 "Chaos",
-                "ChaosCaching"
+                "ChaosCaching",
+				"Buoyancy",
+				"Water"
             }
 		);
 
diff --git a/Engine/Plugins/ChaosCloth/Source/ChaosCloth/Private/ChaosCloth/ChaosClothConstraints.cpp b/Engine/Plugins/ChaosCloth/Source/ChaosCloth/Private/ChaosCloth/ChaosClothConstraints.cpp
--- a/Engine/Plugins/ChaosCloth/Source/ChaosCloth/Private/ChaosCloth/ChaosClothConstraints.cpp
+++ b/Engine/Plugins/ChaosCloth/Source/ChaosCloth/Private/ChaosCloth/ChaosClothConstraints.cpp
@@ -40,6 +40,7 @@
 #include "PhysicsProxy/PerSolverFieldSystem.h"
 #include "Utils/ClothingMeshUtils.h"
 #include "HAL/IConsoleManager.h"
+#include "BuoyancyField.h"
 
 namespace Chaos {
 
@@ -55,6 +56,13 @@ bool bDoPlanarCollisionConstraints = false;
 Softs::FSolverReal PlanarCollisionTol = 0.f;
 Softs::FSolverReal PlanarCollisionStiffness = 1e10;
 
+// #todo(dmp): once we decide what we want to do for enabling buoyancy by default or controlling it some way, we can 
+// enable it for shipping
+namespace ChaosClothConstraintsBuoyancyConsoleVariables
+{
+	bool bEnableBuoyancy = false;
+}
+
 #if !UE_BUILD_SHIPPING
     static FAutoConsoleVariableRef CVarClothbEnableGS(TEXT("p.Chaos.Cloth.EnableGaussSeidel"), bEnableGS, TEXT("Use Gauss Seidel constraints instead of XPBD [def: false]"));
 
@@ -115,6 +123,8 @@ Softs::FSolverReal PlanarCollisionStiffness = 1e10;
 	static FAutoConsoleVariableRef CVarClothPlanarCollisionTol(TEXT("p.Chaos.Cloth.PlanarCollisionTol"), PlanarCollisionTol, TEXT("Collision tolerance for XPBD planar constraint [def: 0.]"));
 
 	static FAutoConsoleVariableRef CVarClothPlanarCollisionStiffness(TEXT("p.Chaos.Cloth.PlanarCollisionStiffness"), PlanarCollisionStiffness, TEXT("Collision stiffness for XPBD planar constraint [def: 1e10.]"));
+	
+	static FAutoConsoleVariableRef CVarClothEnableBuoyancy(TEXT("p.Chaos.Cloth.EnableBuoyancy"), ChaosClothConstraintsBuoyancyConsoleVariables::bEnableBuoyancy, TEXT("Query water bodies for drag and buoyancy effects with cloth [def: false]"));
 #endif
 
 class FClothConstraints::FRuleCreator
@@ -825,8 +835,16 @@ void FClothConstraints::AddRules(
 		// External Forces
 		CreateExternalForces(ConfigProperties, WeightMaps);
 
-		// Velocity Field
-		CreateVelocityAndPressureField(ConfigProperties, WeightMaps, TriangleMesh);
+		if (!ChaosClothConstraintsBuoyancyConsoleVariables::bEnableBuoyancy)
+		{
+			// Velocity Field
+			CreateVelocityAndPressureField(ConfigProperties, WeightMaps, TriangleMesh);
+		}
+		else
+		{
+			// Buoyancy Field
+			CreateBuoyancyField(ConfigProperties, WeightMaps, TriangleMesh);
+		}
 
 		// PerSolverField
 		if (PerSolverField)
@@ -1480,6 +1498,26 @@ void FClothConstraints::CreateVelocityAndPressureField(
 	}
 }
 
+void FClothConstraints::CreateBuoyancyField(
+	const Softs::FCollectionPropertyConstFacade& ConfigProperties,
+	const TMap<FString, TConstArrayView<FRealSingle>>& WeightMaps,
+	const FTriangleMesh& TriangleMesh)
+{
+	if (Evolution)
+	{
+		// Always create buoyancy field--we allow turning it on via blueprints
+		constexpr Softs::FSolverReal WorldScale = (Softs::FSolverReal)100.;
+		BuoyancyField = MakeShared<FBuoyancyField>(
+			Evolution->GetSoftBodyParticles(ParticleRangeId),
+			&TriangleMesh,
+			ConfigProperties,
+			WeightMaps,
+			WorldScale
+			);
+		++NumExternalForceRules;
+	}
+}
+
 void FClothConstraints::CreateExternalForces(
 	const Softs::FCollectionPropertyConstFacade& ConfigProperties,
 	const TMap<FString, TConstArrayView<FRealSingle>>& WeightMaps)
@@ -1756,6 +1794,11 @@ void FClothConstraints::CreateForceBasedRules(const TSharedPtr<Softs::FMultiResC
 		// TODO Linear System
 	}
 
+	if (BuoyancyField)
+	{
+		RuleCreator.AddExternalForceRule_Apply(BuoyancyField.Get());		
+	}
+
 	if (PerSolverField)
 	{
 		RuleCreator.AddExternalForceRule(
@@ -2482,6 +2525,11 @@ void FClothConstraints::UpdateFromSolver(const FSolverVec3& SolverGravity, bool
 		ExternalForces->SetSolverWind(InSolverWindVelocity, LegacyWindAdaptation);
 	}
 	SolverWindVelocity = InSolverWindVelocity;
+
+	if (BuoyancyField)
+	{
+		BuoyancyField->SetSolverGravity(SolverGravity);
+	}
 }
 
 void FClothConstraints::Update(
@@ -2493,6 +2541,8 @@ void FClothConstraints::Update(
 	Softs::FSolverReal MeshScale,
 	Softs::FSolverReal MaxDistancesScale,
 	const FRotation3& LocalSpaceRotation,
+	const FReal LocalSpaceScale,
+	const FSolverVec3 LocalSpaceLocation,
 	const FRotation3& ReferenceSpaceRotation)
 {
 	if (EdgeConstraints)
@@ -2620,6 +2670,23 @@ void FClothConstraints::Update(
 			ReferenceSpaceRotation
 		);
 	}
+	if (BuoyancyField)
+	{
+		constexpr FSolverReal WorldScale = 100.f;
+		const bool bPointBasedWindDisablesAccurateWind = ClothingSimulationClothConsoleVariables::CVarLegacyDisablesAccurateWind.GetValueOnAnyThread();
+		const bool bEnableAerodynamics = !(bUsePointBasedWindModel && bPointBasedWindDisablesAccurateWind);
+		BuoyancyField->SetPropertiesAndBuoyancy(
+			ConfigProperties,
+			WeightMaps,
+			WorldScale,
+			bEnableAerodynamics,
+			SolverWindVelocity,
+			LocalSpaceRotation,
+			LocalSpaceScale,
+			LocalSpaceLocation,
+			ReferenceSpaceRotation
+		);
+	}
 	if (CollisionConstraint)
 	{
 		static IConsoleVariable* const WriteCCDContacts = IConsoleManager::Get().FindConsoleVariable(TEXT("p.Chaos.PBDEvolution.WriteCCDContacts"));
diff --git a/Engine/Plugins/ChaosCloth/Source/ChaosCloth/Private/ChaosCloth/ChaosClothingSimulationCloth.cpp b/Engine/Plugins/ChaosCloth/Source/ChaosCloth/Private/ChaosCloth/ChaosClothingSimulationCloth.cpp
--- a/Engine/Plugins/ChaosCloth/Source/ChaosCloth/Private/ChaosCloth/ChaosClothingSimulationCloth.cpp
+++ b/Engine/Plugins/ChaosCloth/Source/ChaosCloth/Private/ChaosCloth/ChaosClothingSimulationCloth.cpp
@@ -468,7 +468,7 @@ void FClothingSimulationCloth::FLODData::Update(FClothingSimulationSolver* Solve
 		const FReal LocalSpaceScaleInv = 1. / LocalSpaceScale;
 		const Softs::FSolverReal MeshScale = Cloth->Mesh->GetScale() * LocalSpaceScaleInv;
 		const Softs::FSolverReal MaxDistancesScale = (Softs::FSolverReal)Cloth->MaxDistancesMultiplier;
-		ClothConstraints.Update(Cloth->Config->GetProperties(SolverDatum.LODIndex), WeightMaps, VertexSets, FaceSets, FaceIntMaps, MeshScale, MaxDistancesScale, Solver->GetLocalSpaceRotation(), Cloth->ReferenceSpaceTransform.GetRotation());
+		ClothConstraints.Update(Cloth->Config->GetProperties(SolverDatum.LODIndex), WeightMaps, VertexSets, FaceSets, FaceIntMaps, MeshScale, MaxDistancesScale, Solver->GetLocalSpaceRotation(), LocalSpaceScale, Solver->GetLocalSpaceLocation(), Cloth->ReferenceSpaceTransform.GetRotation());
 
 		if (const TSharedPtr<Softs::FPBDSkinnedTriangleMeshCollisionConstraints>& SkinnedTriangleCollisionConstraint = ClothConstraints.GetSkinnedTriangleCollisionsConstraint())
 		{
diff --git a/Engine/Plugins/ChaosCloth/Source/ChaosCloth/Private/ChaosCloth/ChaosClothingSimulationSolver.cpp b/Engine/Plugins/ChaosCloth/Source/ChaosCloth/Private/ChaosCloth/ChaosClothingSimulationSolver.cpp
--- a/Engine/Plugins/ChaosCloth/Source/ChaosCloth/Private/ChaosCloth/ChaosClothingSimulationSolver.cpp
+++ b/Engine/Plugins/ChaosCloth/Source/ChaosCloth/Private/ChaosCloth/ChaosClothingSimulationSolver.cpp
@@ -16,6 +16,8 @@
 #include "Chaos/SkinnedTriangleMesh.h"
 #include "Chaos/WeightedLatticeImplicitObject.h"
 #include "GeometryCollection/ManagedArrayCollection.h"
+#include "BuoyancyField.h"
+
 #if INTEL_ISPC
 #include "ChaosClothingSimulationSolver.ispc.generated.h"
 #endif
@@ -2514,6 +2516,23 @@ FBoxSphereBounds FClothingSimulationSolver::CalculateBounds() const
 	}
 }
 
+
+void FClothingSimulationSolver::PreSimulate_BuoyancyField_GameThread(float dt, const UWorld* World, const FBox BoundingBox)
+{
+	auto& AllConstraints = GetAllClothConstraints();
+
+	for (auto& Pair : AllConstraints)
+	{
+		int32 Key = Pair.Key;
+		::Chaos::FClothConstraints* Constraints = Pair.Value.Get();
+
+		if (FBuoyancyField *BuoyancyField = Constraints->GetBuoyancyField().Get())
+		{
+			BuoyancyField->PreSimulate_GameThread(dt, World, BoundingBox);
+		}
+	}
+}
+
 uint32 FClothingSimulationSolver::GetNumParticles() const
 {
 	return Evolution ? Evolution->GetParticles().Size() : PBDEvolution->GetParticles().Size();
diff --git a/Engine/Plugins/ChaosCloth/Source/ChaosCloth/Public/ChaosCloth/ChaosClothConstraints.h b/Engine/Plugins/ChaosCloth/Source/ChaosCloth/Public/ChaosCloth/ChaosClothConstraints.h
--- a/Engine/Plugins/ChaosCloth/Source/ChaosCloth/Public/ChaosCloth/ChaosClothConstraints.h
+++ b/Engine/Plugins/ChaosCloth/Source/ChaosCloth/Public/ChaosCloth/ChaosClothConstraints.h
@@ -12,6 +12,8 @@
 #include "Chaos/Deformable/GaussSeidelCorotatedCodimensionalConstraints.h"
 #include "Chaos/PBDExtremeDeformationConstraints.h"
 
+class FBuoyancyField;
+
 struct FManagedArrayCollection;
 
 namespace Chaos
@@ -91,6 +93,8 @@ namespace Chaos
 			Softs::FSolverReal MeshScale,
 			Softs::FSolverReal MaxDistancesScale = (Softs::FSolverReal)1.,
 			const FRotation3& LocalSpaceRotation = FRotation3::Identity,
+			const FReal LocalSpaceScale = (Softs::FSolverReal)1.,
+			const FSolverVec3 LocalSpaceLocation = FSolverVec3::ZeroVector,
 			const FRotation3& ReferenceSpaceRotation = FRotation3::Identity);
 
 		// NOTE: this only does something if using the PBDSolver. Force-based solver constraints are activated automatically
@@ -126,6 +130,9 @@ namespace Chaos
 		const TSharedPtr<Softs::FPBDSelfCollisionSphereConstraints>& GetSelfCollisionSphereConstraints() const
 		{ return SelfCollisionSphereConstraints; }
 		const TSharedPtr<Softs::FVelocityAndPressureField>& GetVelocityAndPressureField() const { return VelocityAndPressureField; }
+		
+		const TSharedPtr<FBuoyancyField>& GetBuoyancyField() const { return BuoyancyField; }
+
 		const TSharedPtr<Softs::FExternalForces>& GetExternalForces() const { return ExternalForces; }
 		const TSharedPtr<Softs::FPBDSoftBodyCollisionConstraint>& GetCollisionConstraint() const { return CollisionConstraint; }
 		const TSharedPtr<Softs::FPBDSkinnedTriangleMeshCollisionConstraints>& GetSkinnedTriangleCollisionsConstraint() const { return SkinnedTriangleCollisionsConstraint; }
@@ -189,6 +196,12 @@ namespace Chaos
 			const TMap<FString, TConstArrayView<FRealSingle>>& WeightMaps,
 			const FTriangleMesh& TriangleMesh
 		);
+		void CreateBuoyancyField(
+			const Softs::FCollectionPropertyConstFacade& ConfigProperties,
+			const TMap<FString, TConstArrayView<FRealSingle>>& WeightMaps,
+			const FTriangleMesh& TriangleMesh
+		);
+
 		void CreateExternalForces(
 			const Softs::FCollectionPropertyConstFacade& ConfigProperties,
 			const TMap<FString, TConstArrayView<FRealSingle>>& WeightMaps
@@ -239,6 +252,7 @@ namespace Chaos
 		//~ Begin Force-based solver only constraints
 		Softs::FSolverVec3 SolverWindVelocity; // Set from solver and added to wind from the config
 		TSharedPtr<Softs::FVelocityAndPressureField> VelocityAndPressureField;
+		TSharedPtr<FBuoyancyField> BuoyancyField;
 		TSharedPtr<Softs::FExternalForces> ExternalForces;
 		TSharedPtr<Softs::FPBDSoftBodyCollisionConstraint> CollisionConstraint;
 		TSharedPtr<Softs::FPBDSkinnedTriangleMeshCollisionConstraints> SkinnedTriangleCollisionsConstraint;
diff --git a/Engine/Plugins/ChaosCloth/Source/ChaosCloth/Public/ChaosCloth/ChaosClothingSimulationSolver.h b/Engine/Plugins/ChaosCloth/Source/ChaosCloth/Public/ChaosCloth/ChaosClothingSimulationSolver.h
--- a/Engine/Plugins/ChaosCloth/Source/ChaosCloth/Public/ChaosCloth/ChaosClothingSimulationSolver.h
+++ b/Engine/Plugins/ChaosCloth/Source/ChaosCloth/Public/ChaosCloth/ChaosClothingSimulationSolver.h
@@ -415,6 +415,10 @@ namespace Chaos
 
 		const FClothConstraints& GetClothConstraints(int32 ParticleRangeId) const { return *ClothsConstraints.FindChecked(ParticleRangeId); }
 		FClothConstraints& GetClothConstraints(int32 ParticleRangeId) { return *ClothsConstraints.FindChecked(ParticleRangeId); }
+		const TMap<int32, TUniquePtr<FClothConstraints>> & GetAllClothConstraints() const { return ClothsConstraints; }
+
+		CHAOSCLOTH_API void PreSimulate_BuoyancyField_GameThread(float dt, const UWorld* World, const FBox BoundingBox);		
+
 		CHAOSCLOTH_API uint32 GetNumParticles() const;
 		CHAOSCLOTH_API int32 GetNumActiveParticles() const;
 		CHAOSCLOTH_API int32 GetGlobalParticleOffset(int32 ParticleRangeId) const;
diff --git a/Engine/Plugins/ChaosClothAsset/Source/ChaosClothAssetEngine/Private/ChaosClothAsset/ClothSimulationProxy.cpp b/Engine/Plugins/ChaosClothAsset/Source/ChaosClothAssetEngine/Private/ChaosClothAsset/ClothSimulationProxy.cpp
--- a/Engine/Plugins/ChaosClothAsset/Source/ChaosClothAssetEngine/Private/ChaosClothAsset/ClothSimulationProxy.cpp
+++ b/Engine/Plugins/ChaosClothAsset/Source/ChaosClothAssetEngine/Private/ChaosClothAsset/ClothSimulationProxy.cpp
@@ -454,13 +454,27 @@ namespace UE::Chaos::ClothAsset
 		// Update world forces
 		if (UWorld* const World = ClothComponentAdapter.GetOwnerComponent().GetWorld())
 		{
-			if (UPhysicsFieldComponent* const PhysicsField = World->PhysicsField)
+			static IConsoleVariable* const UsesBuoyancy = IConsoleManager::Get().FindConsoleVariable(TEXT("p.Chaos.Cloth.EnableBuoyancy"));
+			const bool bUsesBuoyancy = UsesBuoyancy ? UsesBuoyancy->GetBool() : false;
+
+			FBox BoundingBox;
+			if (World->PhysicsField || bUsesBuoyancy)
 			{
-				const FBox BoundingBox = CalculateBounds_AnyThread().GetBox().TransformBy(ClothComponentAdapter.GetOwnerComponent().GetComponentTransform());
+				BoundingBox = CalculateBounds_AnyThread().GetBox().TransformBy(ClothComponentAdapter.GetOwnerComponent().GetComponentTransform());
+			}
 
+			if (UPhysicsFieldComponent* const PhysicsField = World->PhysicsField)
+			{
 				PhysicsField->FillTransientCommands(false, BoundingBox, Solver->GetTime(), Solver->GetPerSolverField().GetTransientCommands());
 				PhysicsField->FillPersistentCommands(false, BoundingBox, Solver->GetTime(), Solver->GetPerSolverField().GetPersistentCommands());
 			}
+											
+			// broadphase against water bodies			
+			if (bUsesBuoyancy)
+			{				
+				// find all buoyancy fields and compute collisions				
+				Solver->PreSimulate_BuoyancyField_GameThread(DeltaTime, World, BoundingBox);
+			}
 		}
 
 		// Update external collision sources
@@ -1067,4 +1081,4 @@ namespace UE::Chaos::ClothAsset
 		}
 		return bSolverCurrentlyEnabled;
 	}
-}
+}
\ No newline at end of file
diff --git a/Engine/Plugins/ChaosClothAssetEditor/Source/ChaosClothAssetDataflowNodes/Private/ChaosClothAsset/SimulationAerodynamicsConfigNode.cpp b/Engine/Plugins/ChaosClothAssetEditor/Source/ChaosClothAssetDataflowNodes/Private/ChaosClothAsset/SimulationAerodynamicsConfigNode.cpp
--- a/Engine/Plugins/ChaosClothAssetEditor/Source/ChaosClothAssetDataflowNodes/Private/ChaosClothAsset/SimulationAerodynamicsConfigNode.cpp
+++ b/Engine/Plugins/ChaosClothAssetEditor/Source/ChaosClothAssetDataflowNodes/Private/ChaosClothAsset/SimulationAerodynamicsConfigNode.cpp
@@ -60,4 +60,11 @@ void FChaosClothAssetSimulationAerodynamicsConfigNode::AddProperties(FPropertyHe
 				return ClothFacade.GetSolverAirDamping();
 			}, {});
 	}
+		
+	if (bUseWithWaterBodies)
+	{
+		PropertyHelper.SetProperty(this, &BuoyancyThickness);
+		PropertyHelper.SetProperty(this, &WaterDensity);
+		PropertyHelper.SetProperty(this, &WaterTurbulenceRatio);
+	}
 }
diff --git a/Engine/Plugins/ChaosClothAssetEditor/Source/ChaosClothAssetDataflowNodes/Public/ChaosClothAsset/SimulationAerodynamicsConfigNode.h b/Engine/Plugins/ChaosClothAssetEditor/Source/ChaosClothAssetDataflowNodes/Public/ChaosClothAsset/SimulationAerodynamicsConfigNode.h
--- a/Engine/Plugins/ChaosClothAssetEditor/Source/ChaosClothAssetDataflowNodes/Public/ChaosClothAsset/SimulationAerodynamicsConfigNode.h
+++ b/Engine/Plugins/ChaosClothAssetEditor/Source/ChaosClothAssetDataflowNodes/Public/ChaosClothAsset/SimulationAerodynamicsConfigNode.h
@@ -94,6 +94,30 @@ public:
 	UPROPERTY(EditAnywhere, Category = "Aerodynamics", Meta = (UIMin = "0", UIMax = "1", ClampMin = "0", ClampMax = "10", EditCondition = "bEnableOuterLift", InteractorName = "OuterLift"))
 	FChaosClothAssetWeightedValue OuterLift = { true, 0.035f, 1.f, TEXT("OuterLift"), true };
 	
+	/*
+	 * If true, drag and buoyancy will be used when cloth is submerged under water bodies  
+	 */
+	UPROPERTY(EditAnywhere, Category = "Aerodynamics", Meta = (InteractorName = "UseWithWaterBodies"))
+	bool bUseWithWaterBodies = false;
+	
+	/* 
+	 * Thickness for cloth when used for buoyancy calculations
+	 */
+	UPROPERTY(EditAnywhere, Category = "Aerodynamics", Meta = (UIMin = "0", UIMax = "1", InteractorName = "BuoyancyThickness", EditCondition = "bUseWithWaterBodies"))
+	float BuoyancyThickness = .00001;
+
+	/* 
+	 * Density of fluid underwater
+	 */
+	UPROPERTY(EditAnywhere, Category = "Aerodynamics", Meta = (UIMin = "0", UIMax = "10", ClampMin = "0", ClampMax = "10000", InteractorName = "WaterDensity", EditCondition = "bUseWithWaterBodies"))
+	float WaterDensity = 1000.0f;
+
+	/* 
+	 * Turbulence of fluid underwater
+	 */
+	UPROPERTY(EditAnywhere, Category = "Aerodynamics", Meta = (ClampMin = "0", ClampMax = "1", InteractorName = "WaterTurbulenceRatio", EditCondition = "bUseWithWaterBodies"))
+	float WaterTurbulenceRatio = 0.f;
+
 	FChaosClothAssetSimulationAerodynamicsConfigNode(const UE::Dataflow::FNodeParameters& InParam, FGuid InGuid = FGuid::NewGuid());
 
 private:
diff --git a/Engine/Plugins/Experimental/Buoyancy/Source/Runtime/Private/BuoyancyField.cpp b/Engine/Plugins/Experimental/Buoyancy/Source/Runtime/Private/BuoyancyField.cpp
new file mode 100644
--- /dev/null
+++ b/Engine/Plugins/Experimental/Buoyancy/Source/Runtime/Private/BuoyancyField.cpp
@@ -0,0 +1,326 @@
+// Copyright Epic Games, Inc. All Rights Reserved.
+
+#include "BuoyancyField.h"
+#include "HAL/IConsoleManager.h"
+#include "BuoyancySubsystem.h"
+#include "BuoyancyAlgorithms.h"
+#include "Engine/World.h"
+#include "DrawDebugHelpers.h"
+#include "Chaos/DebugDrawQueue.h"
+
+using namespace Chaos;
+using namespace Chaos::Softs;
+
+extern bool bBuoyancyDebugDraw;
+
+
+
+// #todo(dmp): resolve threading issues so this can work with Chaos Cloth.  See notes below.  For now, this
+// will be left unused and instead we'll only do thread safe GT queries
+bool bComputeWaterPerTri = false;
+// FAutoConsoleVariableRef CVarBuoyancyAlgorithmsFieldComputeWaterPerTri(TEXT("p.Buoyancy.Algorithms.Field.ComputeWaterPerTri"), bComputeWaterPerTri, TEXT(""));
+
+bool bUseAccurateSubmersion = false;
+FAutoConsoleVariableRef CVarBuoyancyAlgorithmsFieldUseAccurateSubmersion(TEXT("p.Buoyancy.Algorithms.Field.UseAccurateSubmersion"), bUseAccurateSubmersion, TEXT(""));
+
+namespace BuoyancyPrivate
+{
+	static float VelocityFieldMaxVelocity = 0.f;
+	static FAutoConsoleVariableRef CVarChaosVelocityFieldMaxVelocity(TEXT("p.Buoyancy.BuoyancyField.MaxVelocity"), VelocityFieldMaxVelocity, TEXT("The maximum relative velocity to process the aerodynamics forces with."));
+}
+
+void FBuoyancyField::SetPropertiesAndBuoyancy(
+	const FCollectionPropertyConstFacade& PropertyCollection,
+	const TMap<FString, TConstArrayView<FRealSingle>>& WeightMaps,
+	FSolverReal WorldScale,
+	bool bEnableAerodynamics,
+	const FSolverVec3& SolverWind,
+	const FRotation3& LocalSpaceRotation,
+	const Chaos::Softs::FSolverReal InLocalSpaceScale,
+	const Chaos::FSolverVec3 InLocalSpaceLocation,
+	const FRotation3& ReferenceSpaceRotation)
+{
+	SetPropertiesAndWind(PropertyCollection, WeightMaps, WorldScale, bEnableAerodynamics, SolverWind, LocalSpaceRotation, ReferenceSpaceRotation);
+	
+	// if none are set, then we disable buoyancy
+	bUseWaterBodies = 
+		BuoyancyThicknessIndex != INDEX_NONE && 
+		WaterDensityIndex != INDEX_NONE && 
+		WaterTurbulenceRatioIndex != INDEX_NONE;
+
+	BuoyancyTriThickness = BuoyancyThicknessIndex != INDEX_NONE ? float(GetBuoyancyThickness(PropertyCollection)) : 0.1f;	
+	WaterDensity = WaterDensityIndex != INDEX_NONE ? float(GetWaterDensity(PropertyCollection)) : 1000.f;	
+	WaterTurbulenceRatio = WaterTurbulenceRatioIndex != INDEX_NONE ? float(GetWaterTurbulenceRatio(PropertyCollection)) : 0.f;	
+
+	WaterRho = FMath::Max(WaterDensity / FMath::Cube(WorldScale), (FSolverReal)0.);
+	
+	constexpr FSolverReal OneQuarter = (FSolverReal)0.25;
+	QuarterWaterRho = WaterRho * OneQuarter;
+
+	LocalSpaceScale = InLocalSpaceScale;
+	LocalSpaceLocation = InLocalSpaceLocation;
+}
+
+bool FBuoyancyField::PreSimulate_GameThread(float DeltaTime, const UWorld* World, const FBox BoundingBox)
+{
+	check(IsInGameThread());
+	
+	TRACE_CPUPROFILER_EVENT_SCOPE(FBuoyancyField_PreSimulate_GameThread);
+
+	UBuoyancySubsystem* BuoyancySubsystem = World->GetSubsystem<UBuoyancySubsystem>();
+	if (bUseWaterBodies && BuoyancySubsystem)
+	{
+		// given the bounding box, find the colliding water bodies
+		TArray<const TSharedPtr<FBuoyancyWaterSplineData>> WaterBodyCollisionData;
+
+		FVector WaterPlaneLocation = {0,0,0};
+		FVector WaterPlaneNormal = {0,0,0};
+		FVector WaterPlaneVelocity = {0,0,0};
+
+		bool ComputePlane = !bComputeWaterPerTri;
+
+		bool HasCollision = BuoyancySubsystem->FindOverlappingWaterBodies(BoundingBox, ComputePlane,
+			WaterPlaneLocation,
+			WaterPlaneNormal,
+			WaterPlaneVelocity,
+			WaterBodyCollisionData);
+
+		WaterBodyCollisions = FBuoyancyCollisionData(BuoyancySubsystem, WaterBodyCollisionData, BoundingBox,
+			HasCollision,
+			ComputePlane,
+			WaterPlaneLocation,
+			WaterPlaneNormal,
+			WaterPlaneVelocity);
+	}
+
+	return true;
+}
+
+void FBuoyancyField::Apply(FSolverParticlesRange& InParticles, const FSolverReal Dt) const
+{
+	TRACE_CPUPROFILER_EVENT_SCOPE(FBuoyancyField_Apply);
+			
+	if (!bUseWaterBodies || WaterBodyCollisions.TheBuoyancySubsystem == nullptr || !WaterBodyCollisions.HasCollision)
+	{
+		FVelocityAndPressureField::Apply(InParticles, Dt);
+		return;
+	}
+	else
+	{
+		const FSolverReal MaxVelocitySquared = (BuoyancyPrivate::VelocityFieldMaxVelocity > 0.f) ? FMath::Square((FSolverReal)BuoyancyPrivate::VelocityFieldMaxVelocity) : TNumericLimits<FSolverReal>::Max();
+
+		// update according to water bodies
+		const TConstArrayView<FSolverVec3>& Xs = InParticles.XArray();
+		const TConstArrayView<FSolverVec3>& Vs = InParticles.GetV();
+		
+		// compute water properties at tri center
+		FVector WaterVelocity;
+		FVector WaterN;
+		FVector WaterPlanePos;
+
+		TSharedPtr<FBuoyancyWaterSplineData> WaterBodySplineData = nullptr;
+
+		if (!bComputeWaterPerTri)
+		{
+			TRACE_CPUPROFILER_EVENT_SCOPE(FBuoyancyField_Apply_ComputeWaterPerCloth);
+			const FVector WorldAvgX = WaterBodyCollisions.WorldBounds.GetCenter();
+
+			WaterVelocity = WaterBodyCollisions.WaterPlaneVelocity;
+			WaterN = WaterBodyCollisions.WaterPlaneNormal;
+			WaterPlanePos = WaterBodyCollisions.WaterPlaneLocation;			
+		}
+		else if (WaterBodyCollisions.WaterBodyCollisionData.Num() > 0)
+		{
+			// #todo(dmp): we aren't compositing submersions from multiple simultaneous bodies right now- just take the first one
+			WaterBodySplineData = WaterBodyCollisions.WaterBodyCollisionData[0];
+		}
+		else
+		{
+			FVelocityAndPressureField::Apply(InParticles, Dt);
+			return;
+		}
+
+		for (int32 ElementIndex = 0; ElementIndex < Elements.Num(); ++ElementIndex)
+		{
+			const TVec3<int32>& Element = Elements[ElementIndex];
+			
+			// lift and drag
+			const FSolverReal CdI = Drag.GetValue(ElementIndex);
+			const FSolverReal CdO = OuterDrag.GetValue(ElementIndex);
+			const FSolverReal ClI = Lift.GetValue(ElementIndex);
+			const FSolverReal ClO = OuterLift.GetValue(ElementIndex);
+			const FSolverReal Cp = Pressure.GetValue(ElementIndex);
+
+			// Compute triangle barycentric center
+			const FSolverVec3 WorldTriPositions[3] = 
+			{
+				Xs[Element[0]] * LocalSpaceScale + LocalSpaceLocation,
+				Xs[Element[1]] * LocalSpaceScale + LocalSpaceLocation,
+				Xs[Element[2]] * LocalSpaceScale + LocalSpaceLocation
+			};
+
+			const FVector WorldTriCenter = FVector((WorldTriPositions[0] + WorldTriPositions[1] + WorldTriPositions[2]) / 3.f);
+
+			// optionally compute water properties at tri center			
+			if (bComputeWaterPerTri)
+			{				
+				TRACE_CPUPROFILER_EVENT_SCOPE(FBuoyancyField_Apply_ComputeWaterPerTri);
+												
+				WaterBodyCollisions.TheBuoyancySubsystem->QueryWaterBody(WorldTriCenter, WaterBodySplineData, WaterVelocity, WaterN, WaterPlanePos);
+			}
+
+			const FSolverVec3 WaterNSolverVec = FSolverVec3(WaterN);
+			const FSolverReal TriVertDotWaterPlane[3] =
+			{
+				(WorldTriPositions[0] - WaterPlanePos).Dot(WaterNSolverVec),
+				(WorldTriPositions[1] - WaterPlanePos).Dot(WaterNSolverVec),
+				(WorldTriPositions[2] - WaterPlanePos).Dot(WaterNSolverVec)
+			};
+
+			// true if the vertex is submerged under the water plane
+			const bool TriVertSubmerged[3] =
+			{
+				TriVertDotWaterPlane[0] < SMALL_NUMBER,
+				TriVertDotWaterPlane[1] < SMALL_NUMBER,
+				TriVertDotWaterPlane[2] < SMALL_NUMBER
+			};
+
+			// early out with no vertices submerged
+			// consider tri as above water and use normal wind and drag
+			if (!TriVertSubmerged[0] && !TriVertSubmerged[1] && !TriVertSubmerged[2])
+			{				
+				const FSolverVec3 Force = FVelocityAndPressureField::CalculateForce(Xs, Vs, ElementIndex);
+				InParticles.Acceleration(Elements[ElementIndex][0]) += InParticles.InvM(Elements[ElementIndex][0]) * Force;
+				InParticles.Acceleration(Elements[ElementIndex][1]) += InParticles.InvM(Elements[ElementIndex][1]) * Force;
+				InParticles.Acceleration(Elements[ElementIndex][2]) += InParticles.InvM(Elements[ElementIndex][2]) * Force;
+
+				continue;
+			}
+			
+			// Force to apply to triangle
+			FSolverVec3 Force = FSolverVec3::ZeroVector;
+
+			// we have submersion if there is an intersection or the tri center is below the water plane
+			// proceed with at least 1 vertex submerged. It can either be fully submerged or intersecting the water surface
+			const bool TriIsFullySubmerged = TriVertSubmerged[0] && TriVertSubmerged[1] && TriVertSubmerged[2];			
+						
+			// Calculate the normal and the area of the surface exposed to the flow
+			FSolverVec3 TriNormal = FSolverVec3::CrossProduct(
+				Xs[Element[2]] - Xs[Element[0]],
+				Xs[Element[1]] - Xs[Element[0]]);
+
+			const FSolverReal NormalLength = TriNormal.SafeNormalize();
+			
+			FSolverReal SubmergedTriArea;
+
+			// In simple submersion, any triangle that is submerged is considered fully submerged
+			if (NormalLength < SMALL_NUMBER)
+			{
+				continue;
+			}
+			else if (!bUseAccurateSubmersion || TriIsFullySubmerged)
+			{
+				// #todo(dmp): simple submersion could be improved by computing the circle circumscribed by the tri then
+				// doing analytical submersion depth based on sphere/plane penetration depth
+				// or
+				// take the bounding box of submerged/intersection points and compute area as a coarse approx
+				// or
+				// take spheres with radius R at each vertex and compute penetration depth of each one and sum up
+				// lots of ways to do this, but plenty of errors due to stretched out triangles misrepresenting the tri area
+				// and it is unclear if it'd be faster than the accurate geometric method below
+				
+				SubmergedTriArea = (FSolverReal) .5 * NormalLength;
+			}
+			else
+			{
+				TRACE_CPUPROFILER_EVENT_SCOPE(FBuoyancyField_Apply_Accurate_Submersion);
+
+				const bool TriVertOnWaterPlane[3] =
+				{
+					TriVertSubmerged[0] && TriVertDotWaterPlane[0] > -SMALL_NUMBER,
+					TriVertSubmerged[1] && TriVertDotWaterPlane[1] > -SMALL_NUMBER,
+					TriVertSubmerged[2] && TriVertDotWaterPlane[2] > -SMALL_NUMBER
+				};
+				
+				// Triangle is partially submerged, compute intersection area
+
+				// vertices composing the submerged polygon (3 or 4 for a triangle)
+				Chaos::FVec3 SubmergedVerts[4];								
+				int NumSubmergedVerts = 0;
+				
+				// compute the ordered submerged vertices
+				for (int i = 0; i < 3; ++i)
+				{
+					if (TriVertSubmerged[i])
+					{
+						SubmergedVerts[NumSubmergedVerts++] = WorldTriPositions[i];
+					}
+						
+					// #todo(dmp) inline?
+					Chaos::FVec3 EdgeIntersection;
+					if (!TriVertOnWaterPlane[(i+1) % 3] && BuoyancyAlgorithms::EdgePlaneIntersection(WaterPlanePos, WaterN, WorldTriPositions[i], WorldTriPositions[(i+1) % 3], EdgeIntersection))						
+					{
+						SubmergedVerts[NumSubmergedVerts++] = EdgeIntersection;
+					}
+				}
+
+				// compute the submerged volume constructed as 1 or two triangles
+				// num submerged is either 3 (submerged triangle) or 4 (submerged quad).  
+				// otherwise we have a degenerate case
+				if (NumSubmergedVerts < 3)
+				{
+					continue;
+				}
+
+				// Compute area of first triangle
+				SubmergedTriArea = (FSolverReal) .5 * 
+					(SubmergedVerts[2] - SubmergedVerts[0]).Cross(
+					SubmergedVerts[1] - SubmergedVerts[0]).Length();
+					
+				// two triangle that span the submerged part, add the second one
+				if (NumSubmergedVerts == 4)
+				{
+					SubmergedTriArea += (FSolverReal) .5 * 
+						(SubmergedVerts[2] - SubmergedVerts[0]).Cross(
+						SubmergedVerts[3] - SubmergedVerts[0]).Length();
+				}					
+			}						
+			
+			// project water velocity onto water plane
+			const FSolverVec3 WaterVelocityOnPlane = WaterVelocity - WaterVelocity.Dot(WaterN) * WaterN;
+										
+			// Tri velocity average of 3 verts
+			const FSolverVec3& TriSurfaceVelocity = (FSolverReal)(1. / 3.) * (
+				Vs[Element[0]] +
+				Vs[Element[1]] +
+				Vs[Element[2]]);
+			
+			Force = CalculateForce(TriNormal, TriSurfaceVelocity, SubmergedTriArea, WaterVelocityOnPlane, CdI, CdO, ClI, ClO, Cp, MaxVelocitySquared);
+									
+			// buoyancy force		
+			const FSolverReal SubmergedTriVolume = SubmergedTriArea * BuoyancyTriThickness;			
+			//                               cm/sec^2     *    kg/cm^3          *     cm^3
+		    Force += ((FSolverReal) -1.0) * SolverGravity * WaterRho * SubmergedTriVolume;			
+
+			// increment acceleration on each vertex
+			// #todo(dmp): only update vertices that are submerged or whole triangle?
+			InParticles.Acceleration(Element[0]) += InParticles.InvM(Element[0]) * Force;
+			InParticles.Acceleration(Element[1]) += InParticles.InvM(Element[1]) * Force;
+			InParticles.Acceleration(Element[2]) += InParticles.InvM(Element[2]) * Force;
+
+#if ENABLE_DRAW_DEBUG
+			if (bBuoyancyDebugDraw)
+			{
+				const FVector WorldAvgX = WaterBodyCollisions.WorldBounds.GetCenter();
+
+				Chaos::FDebugDrawQueue::GetInstance().DrawDebugDirectionalArrow(WorldTriCenter, WorldTriCenter + FVector(TriNormal)*10, 1, FColor::Black, false, -1.f, -1, .25f);
+				Chaos::FDebugDrawQueue::GetInstance().DrawDebugSphere(WaterPlanePos, 3, 10, FColor::Orange, false, -1.f, -1, 2.0f);
+				Chaos::FDebugDrawQueue::GetInstance().DrawDebugSphere(WorldAvgX, 3, 10, FColor::Blue, false, -1.f, -1, 2.0f);
+				Chaos::FDebugDrawQueue::GetInstance().DrawDebugDirectionalArrow(WaterPlanePos, WaterPlanePos + WaterN * 100, 20, FColor::Orange, false, -1.f, -1, 6.f);
+				Chaos::FDebugDrawQueue::GetInstance().DrawDebugDirectionalArrow(WaterPlanePos, WaterPlanePos + WaterVelocity * 2, 1, FColor::Green, false, -1.f, -1, 6.f);
+				Chaos::FDebugDrawQueue::GetInstance().DrawDebugDirectionalArrow(WorldTriCenter, WorldTriCenter + FVector(Force), 1, FColor::Blue, false, -1.f, -1, 1.f);
+			}
+#endif				
+		}
+	}
+}
\ No newline at end of file
diff --git a/Engine/Plugins/Experimental/Buoyancy/Source/Runtime/Private/BuoyancySubsystem.cpp b/Engine/Plugins/Experimental/Buoyancy/Source/Runtime/Private/BuoyancySubsystem.cpp
--- a/Engine/Plugins/Experimental/Buoyancy/Source/Runtime/Private/BuoyancySubsystem.cpp
+++ b/Engine/Plugins/Experimental/Buoyancy/Source/Runtime/Private/BuoyancySubsystem.cpp
@@ -344,8 +344,6 @@ void UBuoyancySubsystem::UpdateNetMode()
 	}
 }
 
-// #todo(dmp): the current use of this method is on whatever thread Cloth is evaluated with, which is not the physics thread,
-//  but we are using the FBuoyancyWaterSplineData, which is stored on the PT after it is created from the Buoyancy plugin
 bool UBuoyancySubsystem::QueryWaterBody(const FVector& InputPosition, const TSharedPtr<FBuoyancyWaterSplineData> WaterData,
 	FVector& WaterVel, FVector& WaterPlaneN, FVector& WaterPlanePos)
 {		
@@ -398,9 +396,16 @@ bool UBuoyancySubsystem::QueryWaterBody(const FVector& InputPosition, const TSha
 	return false;
 }
 
-bool UBuoyancySubsystem::FindOverlappingWaterBodies(const FBox BoundingBox, TArray<const TSharedPtr<FBuoyancyWaterSplineData>>& WaterBodyPhysicsProxies)
+bool UBuoyancySubsystem::FindOverlappingWaterBodies(const FBox BoundingBox, 
+	bool ComputePlane, 
+	FVector &ClosestWaterPlaneLocation,
+	FVector &ClosestWaterPlaneNormal,
+	FVector &ClosestWaterPlaneVelocity,
+	TArray<const TSharedPtr<FBuoyancyWaterSplineData>>& CollisionWaterBodySplineData)
 {
-	if (!SplineData)
+	check(IsInGameThread());
+
+	if (!SplineData && !ComputePlane)
 	{
 		return false;
 	}
@@ -416,7 +421,7 @@ bool UBuoyancySubsystem::FindOverlappingWaterBodies(const FBox BoundingBox, TArr
 
 	TArray<FOverlapResult> Overlaps;
 			
-	// perform overlap test
+	// perform overlap test to find collisions with the input bounding box
 	bool bHit = GetWorld()->OverlapMultiByObjectType(
 		Overlaps,
 		BoundingBox.GetCenter(),
@@ -426,40 +431,78 @@ bool UBuoyancySubsystem::FindOverlappingWaterBodies(const FBox BoundingBox, TArr
 		QueryParams
 	);
 
-	WaterBodyPhysicsProxies.Empty();
+	CollisionWaterBodySplineData.Empty();
 				
 	bool bFoundWaterBody = false;
 
 	// if we have a hit, find the water bodies we are colliding with
 	if (bHit)
 	{
+		float MinDistance = TNumericLimits <float>::Max();
+
 		for (const FOverlapResult& Result : Overlaps)
 		{
 			UWaterBodyComponent* WaterBodyComponent = Cast<UWaterBodyComponent>(Result.Component->GetAttachParent());
 			if (WaterBodyComponent)
-			{						
-				for (UPrimitiveComponent* WaterPrimitiveComponent : WaterBodyComponent->GetCollisionComponents(true))
+			{	
+				if (ComputePlane)
 				{
-					// Add each object (probably just one) to the objects list
-					if (WaterPrimitiveComponent)
+					const TValueOrError<FWaterBodyQueryResult, EWaterBodyQueryError> QueryResult = 
+						WaterBodyComponent->TryQueryWaterInfoClosestToWorldLocation(BoundingBox.GetCenter(),
+							EWaterBodyQueryFlags::ComputeLocation
+							| EWaterBodyQueryFlags::ComputeNormal
+							| EWaterBodyQueryFlags::ComputeVelocity
+							| EWaterBodyQueryFlags::ComputeDepth);
+
+					if (!QueryResult.HasValue())
 					{
-						FBodyInstance& BodyInstance = WaterPrimitiveComponent->BodyInstance;
-						if (BodyInstance.IsValidBodyInstance())
+						continue;
+					}
+
+					const FWaterBodyQueryResult &WaterCollisionResult = QueryResult.GetValue();
+
+					const FVector WaterPlaneLocation = WaterCollisionResult.GetWaterPlaneLocation();
+					const FVector WaterPlaneNormal = WaterCollisionResult.GetWaterPlaneNormal();
+					const FVector WaterPlaneVelocity = WaterCollisionResult.GetVelocity();
+
+					const float DistanceToInput = FVector::Distance(WaterPlaneLocation, BoundingBox.GetCenter());
+
+					if (DistanceToInput < MinDistance && WaterCollisionResult.GetWaterPlaneDepth() > SMALL_NUMBER)
+					{
+						bFoundWaterBody = true;
+
+						MinDistance = DistanceToInput;
+
+						ClosestWaterPlaneLocation = WaterPlaneLocation;
+						ClosestWaterPlaneNormal = WaterPlaneNormal;
+						ClosestWaterPlaneVelocity = WaterPlaneVelocity;
+					}
+				}
+				else
+				{
+					for (UPrimitiveComponent* WaterPrimitiveComponent : WaterBodyComponent->GetCollisionComponents(true))
+					{
+						// Add each object (probably just one) to the objects list
+						if (WaterPrimitiveComponent)
 						{
-							Chaos::FSingleParticlePhysicsProxy* WaterProxy = BodyInstance.GetPhysicsActor();
-							
-							// if we found a valid physics actor, get the PT spline data associated with the particle
-							// #todo(dmp): resolve issues with threading since this method is called from the GT but we are
-							// querying PT data here that is created from the Buoyancy plugin.
-							if (WaterProxy)
+							FBodyInstance& BodyInstance = WaterPrimitiveComponent->BodyInstance;
+							if (BodyInstance.IsValidBodyInstance())
 							{
-								bFoundWaterBody = true;
-								Chaos::FRigidBodyHandle_External& Body_External = WaterProxy->GetGameThreadAPI();								
+								Chaos::FSingleParticlePhysicsProxy* WaterProxy = BodyInstance.GetPhysicsActor();
+							
+								// if we found a valid physics actor, get the PT spline data associated with the particle
+								// #todo(dmp): resolve issues with threading since this method is called from the GT but we are
+								// querying PT data here that is created from the Buoyancy plugin.
+								if (WaterProxy)
+								{									
+									Chaos::FRigidBodyHandle_External& Body_External = WaterProxy->GetGameThreadAPI();								
 								
-								const TSharedPtr<FBuoyancyWaterSplineData> *TmpData = SplineData->GetData_PT(WaterProxy->GetGameThreadAPI());								
-								if (TmpData)
-								{
-									WaterBodyPhysicsProxies.Add(*TmpData);
+									const TSharedPtr<FBuoyancyWaterSplineData> *TmpData = SplineData->GetData_PT(WaterProxy->GetGameThreadAPI());								
+									if (TmpData)
+									{
+										bFoundWaterBody = true;
+										CollisionWaterBodySplineData.Add(*TmpData);
+									}
 								}
 							}
 						}
diff --git a/Engine/Plugins/Experimental/Buoyancy/Source/Runtime/Public/BuoyancyField.h b/Engine/Plugins/Experimental/Buoyancy/Source/Runtime/Public/BuoyancyField.h
new file mode 100644
--- /dev/null
+++ b/Engine/Plugins/Experimental/Buoyancy/Source/Runtime/Public/BuoyancyField.h
@@ -0,0 +1,105 @@
+// Copyright Epic Games, Inc. All Rights Reserved.
+#pragma once
+
+#include "Chaos/VelocityField.h"
+#include "BuoyancySubsystem.h"
+
+using Chaos::Softs::FSolverVec3;
+using Chaos::Softs::FSolverReal;
+
+// Velocity field used solely for aerodynamics effects including water submersion
+class FBuoyancyField : public Chaos::Softs::FVelocityAndPressureField
+{
+public:
+	
+	explicit FBuoyancyField(const Chaos::Softs::FCollectionPropertyConstFacade& PropertyCollection) : 
+		Chaos::Softs::FVelocityAndPressureField(PropertyCollection),
+		BuoyancyThicknessIndex(PropertyCollection),
+		WaterDensityIndex(PropertyCollection),
+		WaterTurbulenceRatioIndex(PropertyCollection)
+	{
+	}
+
+	FBuoyancyField(
+		const Chaos::Softs::FSolverParticlesRange& Particles,
+		const Chaos::FTriangleMesh* TriangleMesh,
+		const Chaos::Softs::FCollectionPropertyConstFacade& PropertyCollection,
+		const TMap<FString, TConstArrayView<Chaos::FRealSingle>>& Weightmaps,
+		Chaos::FSolverReal WorldScale)
+		: Chaos::Softs::FVelocityAndPressureField(Particles, TriangleMesh, PropertyCollection, Weightmaps, WorldScale),
+		BuoyancyThicknessIndex(PropertyCollection),
+		WaterDensityIndex(PropertyCollection),
+		WaterTurbulenceRatioIndex(PropertyCollection)
+		
+	{		
+	}
+
+	// Construct an uninitialized field. Mesh, properties, and velocity will have to be set for this field to be valid.
+	FBuoyancyField() : Chaos::Softs::FVelocityAndPressureField(),
+		BuoyancyThicknessIndex(ForceInit),
+		WaterDensityIndex(ForceInit),
+		WaterTurbulenceRatioIndex(ForceInit)
+	{
+	}
+
+	~FBuoyancyField() {}
+
+	BUOYANCY_API void Apply(Chaos::Softs::FSolverParticlesRange& InParticles, const Chaos::FSolverReal Dt) const;
+	
+	BUOYANCY_API void SetPropertiesAndBuoyancy(
+		const Chaos::Softs::FCollectionPropertyConstFacade& PropertyCollection,
+		const TMap<FString, TConstArrayView<Chaos::FRealSingle>>& Weightmaps,
+		FSolverReal WorldScale,
+		bool bEnableAerodynamics,
+		const FSolverVec3& SolverWind,
+		const Chaos::FRotation3& LocalSpaceRotation = Chaos::FRotation3::Identity,
+		const Chaos::Softs::FSolverReal LocalSpaceScale = (Chaos::Softs::FSolverReal)1.,
+		const Chaos::FSolverVec3 LocalSpaceLocation = FSolverVec3::ZeroVector,
+		const Chaos::FRotation3& ReferenceSpaceRotation = Chaos::FRotation3::Identity
+		);
+
+	BUOYANCY_API bool PreSimulate_GameThread(float DeltaTime, const UWorld* World, const FBox BoundingBox);
+
+	void SetSolverGravity(FSolverVec3 Gravity) { SolverGravity = Gravity; }
+
+	UE_CHAOS_DECLARE_INDEXLESS_PROPERTYCOLLECTION_NAME(BuoyancyThickness, float);
+	UE_CHAOS_DECLARE_INDEXLESS_PROPERTYCOLLECTION_NAME(WaterDensity, float);
+	UE_CHAOS_DECLARE_INDEXLESS_PROPERTYCOLLECTION_NAME(WaterTurbulenceRatio, float);
+
+protected:
+
+	FSolverVec3 CalculateForce(FSolverVec3 N, const FSolverVec3& SurfaceVelocity, const FSolverReal SubmergedTriArea, const FSolverVec3& FluidVelocity, const FSolverReal CdI, const FSolverReal CdO, const FSolverReal ClI, const FSolverReal ClO, const FSolverReal Cp, const FSolverReal MaxVelocitySquared) const
+	{		
+		checkSlow(MaxVelocitySquared > (FSolverReal)0);
+				
+		// Calculate the direction and the relative velocity of the triangle to the flow
+		
+		FSolverVec3 V = FluidVelocity - SurfaceVelocity;
+
+		// Clamp the velocity
+		const FSolverReal RelVelocitySquared = V.SquaredLength();
+		if (RelVelocitySquared > MaxVelocitySquared)
+		{
+			V *= FMath::Sqrt(MaxVelocitySquared / RelVelocitySquared);
+		}
+
+		return CalculateForceLowLevel(V, N, QuarterWaterRho, (FSolverReal) 2. * SubmergedTriArea, CdI, CdO, ClI, ClO, Cp, WaterTurbulenceRatio);
+	}
+
+	bool bUseWaterBodies = true;
+
+	FBuoyancyCollisionData WaterBodyCollisions;
+	FSolverReal LocalSpaceScale = 1.f;
+	FSolverVec3 LocalSpaceLocation = {0.f,0.f,0.f};
+	FSolverVec3 SolverGravity = {0.f,0.f,0.f};
+	FSolverReal BuoyancyTriThickness = .01f;
+	FSolverReal WaterDensity = 1.f;
+	FSolverReal WaterTurbulenceRatio = 1.f;
+
+	FSolverReal WaterRho = 1.f;
+	FSolverReal QuarterWaterRho = 1.f;
+
+	UE_CHAOS_DECLARE_INDEXED_PROPERTYCOLLECTION_NAME(BuoyancyThickness, float);
+	UE_CHAOS_DECLARE_INDEXED_PROPERTYCOLLECTION_NAME(WaterDensity, float);
+	UE_CHAOS_DECLARE_INDEXED_PROPERTYCOLLECTION_NAME(WaterTurbulenceRatio, float);
+};
\ No newline at end of file
diff --git a/Engine/Plugins/Experimental/Buoyancy/Source/Runtime/Public/BuoyancySubsystem.h b/Engine/Plugins/Experimental/Buoyancy/Source/Runtime/Public/BuoyancySubsystem.h
--- a/Engine/Plugins/Experimental/Buoyancy/Source/Runtime/Public/BuoyancySubsystem.h
+++ b/Engine/Plugins/Experimental/Buoyancy/Source/Runtime/Public/BuoyancySubsystem.h
@@ -78,18 +78,42 @@ struct FBuoyancySettings
 class FBuoyancyCollisionData
 {
 	public:
-		FBuoyancyCollisionData(UBuoyancySubsystem *InTheBuoyancySubsystem, TArray<const TSharedPtr<FBuoyancyWaterSplineData>> InWaterBodyCollisionData)
+		FBuoyancyCollisionData(UBuoyancySubsystem *InTheBuoyancySubsystem, 
+			TArray<const TSharedPtr<FBuoyancyWaterSplineData>> InWaterBodyCollisionData, FBox InWorldBounds,
+			bool InHasCollision,
+			bool InHasPlane,
+			FVector InWaterPlaneLocation,
+			FVector InWaterPlaneNormal,
+			FVector InWaterPlaneVelocity)
 		: TheBuoyancySubsystem(InTheBuoyancySubsystem)
-		, WaterBodyCollisionData(InWaterBodyCollisionData) 
+		, WaterBodyCollisionData(InWaterBodyCollisionData)
+		, WorldBounds(InWorldBounds)
+		, HasCollision(InHasCollision)
+		, HasPlane(InHasPlane)
+		, WaterPlaneLocation(InWaterPlaneLocation)
+		, WaterPlaneNormal(InWaterPlaneNormal)
+		, WaterPlaneVelocity(InWaterPlaneVelocity)
 		{}
 		
 		FBuoyancyCollisionData() 
 		: TheBuoyancySubsystem(nullptr)
 		, WaterBodyCollisionData(TArray<const TSharedPtr<FBuoyancyWaterSplineData>>()) 
+		, WorldBounds()
+		, HasCollision(false)
+		, HasPlane(false)
+		, WaterPlaneLocation(0,0,0)
+		, WaterPlaneNormal(0,0,1)
+		, WaterPlaneVelocity(0,0,0)
 		{}
 
 		UBuoyancySubsystem *TheBuoyancySubsystem;
 		TArray<const TSharedPtr<FBuoyancyWaterSplineData>> WaterBodyCollisionData;
+		FBox WorldBounds;
+		bool HasCollision;
+		bool HasPlane;
+		FVector WaterPlaneLocation;
+		FVector WaterPlaneNormal;
+		FVector WaterPlaneVelocity;
 };
 
 //
@@ -127,18 +151,22 @@ public:
 	UE_API SIZE_T GetAllocatedSize() const;
 #endif
 
-	// given a bounding box, perform an overlap test and filter for water bodies.  There could be some potential threading
-	// issues here since we are querying the FBuoyancyWaterSplineDataManager physics thread data presumably from the gt to
-	// do the overlap test.  We need a better way of communicating the buoyancy data so we can capture it on the gt, 
-	// then use it on the physics or other thread (ie: cloth).
+	// Given a bounding box, perform an overlap test and filter for water bodies, returning the data needed for buoyancy.  
+	// There are threading issues here since we are storing pointers to FBuoyancyWaterSplineDataManager
+	// physics thread data on the gt to use for tests later on (with FindOverlappingWaterBodies).  
+	// Ideally the Cloth solver would be able to make use of TUserDataManagerPT even if it doesn't run on the pt	
 	UE_API bool FindOverlappingWaterBodies(
-		const FBox BoundingBox, TArray<const TSharedPtr<FBuoyancyWaterSplineData>> &WaterBodyPhysicsProxies);
-
-	// query a water body given a FBuoyancyWaterSplineData object representing a water body
-	// note this is designed to only work on the physics thread, but the FBuoyancyWaterSplineData data is 
-	// generated once and then only written to if water bodies change (uncommon).  In the future, we should consider
-	// restructuring usage such that there aren't potential threading issues, but it is problematic since Chaos and
-	// Cloth are evaluated on separate threads.
+		const FBox BoundingBox, 
+		bool ComputePlane, 
+		FVector &ClosestWaterPlaneLocation,
+		FVector &ClosestWaterPlaneNormal,
+		FVector &ClosestWaterPlaneVelocity,
+		TArray<const TSharedPtr<FBuoyancyWaterSplineData>> &CollisionWaterBodySplineData);
+
+	// Query a water body given a FBuoyancyWaterSplineData object representing a water body
+	// #todo(dmp): This is designed to only work on the physics thread, needs to be more robust if used elsewhere
+	// IE: Ideally the Cloth solver would be able to make use of TUserDataManagerPT even if it doesn't run on the pt
+	// and it'd need to support the spline cache
 	UE_API bool QueryWaterBody(const FVector& InputPosition,  const TSharedPtr<FBuoyancyWaterSplineData> WaterData,
 		FVector& WaterVel, FVector& WaterPlaneN, FVector& WaterPlanePos);
 
diff --git a/Engine/Source/Runtime/Experimental/Chaos/Public/Chaos/VelocityField.h b/Engine/Source/Runtime/Experimental/Chaos/Public/Chaos/VelocityField.h
--- a/Engine/Source/Runtime/Experimental/Chaos/Public/Chaos/VelocityField.h
+++ b/Engine/Source/Runtime/Experimental/Chaos/Public/Chaos/VelocityField.h
@@ -13,7 +13,7 @@ namespace Chaos::Softs
 {
 
 // Velocity field used solely for aerodynamics effects, use Chaos Fields for other types of fields.
-class FVelocityAndPressureField final
+class FVelocityAndPressureField
 {
 public:
 	static constexpr FSolverReal DefaultDragCoefficient = (FSolverReal)0.5;
@@ -108,7 +108,7 @@ public:
 	{
 	}
 
-	~FVelocityAndPressureField() = default;
+	virtual ~FVelocityAndPressureField() {}
 
 	CHAOS_API void UpdateForces(const FSolverParticles& InParticles, const FSolverReal /*Dt*/);
 
@@ -206,7 +206,7 @@ public:
 	UE_CHAOS_DECLARE_INDEXLESS_PROPERTYCOLLECTION_NAME(WindVelocity, FVector3f);
 	UE_CHAOS_DECLARE_INDEXLESS_PROPERTYCOLLECTION_NAME(WindVelocitySpace, int32);
 
-private:
+protected:
 	bool AreAerodynamicsEnabled() const { return QuarterRho > (FSolverReal)0.; }
 
 	CHAOS_API void InitColor(const FSolverParticlesRange& InParticles);
@@ -225,6 +225,20 @@ private:
 		const TConstArrayView<FRealSingle>& PressureMultipliers);
 
 
+	FSolverVec3 CalculateForceLowLevel(const FSolverVec3 &InV, const FSolverVec3 &InN, const FSolverReal InQuarterRho, const FSolverReal InDoubleArea, 
+		FSolverReal InCdI, FSolverReal InCdO, const FSolverReal InClI, const FSolverReal InClO, const FSolverReal InCp, const FSolverReal InTr) const
+	{		
+		// Set the aerodynamic forces
+		const FSolverReal VDotN = FSolverVec3::DotProduct(InV, InN);
+		const FSolverReal VSquare = FSolverVec3::DotProduct(InV, InV);
+
+		const FSolverReal TurbulenceFactor = VSquare > UE_SMALL_NUMBER ? ((1.f - InTr) / FMath::Sqrt(VSquare) + InTr) : 1.f;
+
+		return InQuarterRho * InDoubleArea * TurbulenceFactor * (VDotN >= (FSolverReal)0. ?  // The flow can hit either side of the triangle, so the normal might need to be reversed
+			(InCdI - InClI) * VDotN * InV + InClI * VSquare * InN :
+			(InClO - InCdO) * VDotN * InV - InClO * VSquare * InN) + InDoubleArea * (FSolverReal)0.5 * InCp * InN;
+	}
+
 	FSolverVec3 CalculateForce(const TConstArrayView<FSolverVec3>& Xs, const TConstArrayView<FSolverVec3>& Vs, int32 ElementIndex, const FSolverVec3& InVelocity, const FSolverReal CdI, const FSolverReal CdO, const FSolverReal ClI, const FSolverReal ClO, const FSolverReal Cp) const
 	{
 		const TVec3<int32>& Element = Elements[ElementIndex];
@@ -242,15 +256,7 @@ private:
 			Vs[Element[2]]);
 		const FSolverVec3 V = InVelocity - SurfaceVelocity;
 
-		// Set the aerodynamic forces
-		const FSolverReal VDotN = FSolverVec3::DotProduct(V, N);
-		const FSolverReal VSquare = FSolverVec3::DotProduct(V, V);
-
-		const FSolverReal TurbulenceFactor = VSquare > UE_SMALL_NUMBER ? ((1.f - TurbulenceRatio) / FMath::Sqrt(VSquare) + TurbulenceRatio) : 1.f;
-
-		return QuarterRho * DoubleArea * TurbulenceFactor * (VDotN >= (FSolverReal)0. ?  // The flow can hit either side of the triangle, so the normal might need to be reversed
-			(CdI - ClI) * VDotN * V + ClI * VSquare * N :
-			(ClO - CdO) * VDotN * V - ClO * VSquare * N) + DoubleArea * (FSolverReal)0.5 * Cp * N;
+		return CalculateForceLowLevel(V, N, QuarterRho, DoubleArea, CdI, CdO, ClI, ClO, Cp, TurbulenceRatio);
 	}
 
 	void UpdateField(const FSolverParticles& InParticles, int32 ElementIndex, const FSolverVec3& InVelocity, const FSolverReal CdI, const FSolverReal CdO, const FSolverReal ClI, const FSolverReal ClO, const FSolverReal Cp)
@@ -284,14 +290,8 @@ private:
 			V *= FMath::Sqrt(MaxVelocitySquared / RelVelocitySquared);
 		}
 
-		// Set the aerodynamic forces
-		const FSolverReal VDotN = FSolverVec3::DotProduct(V, N);
-		const FSolverReal VSquare = FSolverVec3::DotProduct(V, V);
-		const FSolverReal TurbulenceFactor = VSquare > UE_SMALL_NUMBER ? ((1.f - TurbulenceRatio) / FMath::Sqrt(VSquare) + TurbulenceRatio) : 1.f;
 
-		return QuarterRho * DoubleArea * TurbulenceFactor * (VDotN >= (FSolverReal)0. ?  // The flow can hit either side of the triangle, so the normal might need to be reversed
-			(CdI - ClI) * VDotN * V + ClI * VSquare * N :
-			(ClO - CdO) * VDotN * V - ClO * VSquare * N) + DoubleArea * (FSolverReal)0.5 * Cp * N;
+		return CalculateForceLowLevel(V, N, QuarterRho, DoubleArea, CdI, CdO, ClI, ClO, Cp, TurbulenceRatio);
 	}
 
 	void UpdateField(const FSolverParticles& InParticles, int32 ElementIndex, const FSolverVec3& InVelocity, const FSolverReal CdI, const FSolverReal CdO, const FSolverReal ClI, const FSolverReal ClO, const FSolverReal Cp, const FSolverReal MaxVelocitySquared)
@@ -299,7 +299,7 @@ private:
 		Forces[ElementIndex] = CalculateForce(TConstArrayView<FSolverVec3>(InParticles.XArray()), TConstArrayView<FSolverVec3>(InParticles.GetV()), ElementIndex, InVelocity, CdI, CdO, ClI, ClO, Cp, MaxVelocitySquared);
 	}
 
-private:
+protected:
 	int32 Offset;
 	int32 NumParticles;
 	TConstArrayView<TArray<int32>> PointToTriangleMap; // Points use global indexing. May point to PointToTriangleMapLocal or data in the original FTriangleMesh
