addressing code review comments/ cleanup

#Synced-CL 50680703

diff --git a/Engine/Plugins/ChaosCloth/ChaosCloth.uplugin b/Engine/Plugins/ChaosCloth/ChaosCloth.uplugin
--- a/Engine/Plugins/ChaosCloth/ChaosCloth.uplugin
+++ b/Engine/Plugins/ChaosCloth/ChaosCloth.uplugin
@@ -33,6 +33,10 @@
 		{
 			"Name": "Buoyancy",
 			"Enabled": true
+		},
+		{
+			"Name": "Water",
+			"Enabled": true
 		}
 	]
 }
diff --git a/Engine/Plugins/ChaosCloth/Source/ChaosCloth/Private/ChaosCloth/ChaosClothConstraints.cpp b/Engine/Plugins/ChaosCloth/Source/ChaosCloth/Private/ChaosCloth/ChaosClothConstraints.cpp
--- a/Engine/Plugins/ChaosCloth/Source/ChaosCloth/Private/ChaosCloth/ChaosClothConstraints.cpp
+++ b/Engine/Plugins/ChaosCloth/Source/ChaosCloth/Private/ChaosCloth/ChaosClothConstraints.cpp
@@ -842,9 +842,9 @@ void FClothConstraints::AddRules(
 			CreateVelocityAndPressureField(ConfigProperties, WeightMaps, TriangleMesh);
 		}
 		else
-		{
+		{		
 			// Buoyancy Field
-			CreateBuoyancyField(ConfigProperties, WeightMaps, TriangleMesh);
+			CreateVelocityPressureAndBuoyancyField(ConfigProperties, WeightMaps, TriangleMesh);
 		}
 
 		// PerSolverField
@@ -1499,21 +1499,22 @@ void FClothConstraints::CreateVelocityAndPressureField(
 	}
 }
 
-void FClothConstraints::CreateBuoyancyField(
+void FClothConstraints::CreateVelocityPressureAndBuoyancyField(
 	const Softs::FCollectionPropertyConstFacade& ConfigProperties,
 	const TMap<FString, TConstArrayView<FRealSingle>>& WeightMaps,
 	const FTriangleMesh& TriangleMesh)
 {
 	if (Evolution)
 	{
-		// Always create buoyancy field--we allow turning it on via blueprints
+		// Always create velocity, pressure and buoyancy field--we allow turning it on via blueprints
 		constexpr Softs::FSolverReal WorldScale = (Softs::FSolverReal)100.;
-		BuoyancyField = MakeShared<FBuoyancyField>(
+		VelocityPressureAndBuoyancyField = MakeShared<Buoyancy::FBuoyancyField>(
 			Evolution->GetSoftBodyParticles(ParticleRangeId),
 			&TriangleMesh,
 			ConfigProperties,
 			WeightMaps,
-			WorldScale
+			WorldScale,
+			ExternalForces
 			);
 		++NumExternalForceRules;
 	}
@@ -1795,9 +1796,9 @@ void FClothConstraints::CreateForceBasedRules(const TSharedPtr<Softs::FMultiResC
 		// TODO Linear System
 	}
 
-	if (BuoyancyField)
+	if (VelocityPressureAndBuoyancyField)
 	{
-		RuleCreator.AddExternalForceRule_Apply(BuoyancyField.Get());		
+		RuleCreator.AddExternalForceRule_Apply(VelocityPressureAndBuoyancyField.Get());		
 	}
 
 	if (PerSolverField)
@@ -2531,11 +2532,6 @@ void FClothConstraints::UpdateFromSolver(const FSolverVec3& SolverGravity, bool
 		ExternalForces->SetSolverWind(InSolverWindVelocity, LegacyWindAdaptation);
 	}
 	SolverWindVelocity = InSolverWindVelocity;
-
-	if (BuoyancyField)
-	{
-		BuoyancyField->SetSolverGravity(SolverGravity);
-	}
 }
 
 void FClothConstraints::Update(
@@ -2547,9 +2543,9 @@ void FClothConstraints::Update(
 	Softs::FSolverReal MeshScale,
 	Softs::FSolverReal MaxDistancesScale,
 	const FRotation3& LocalSpaceRotation,
+	const FRotation3& ReferenceSpaceRotation,
 	const FReal LocalSpaceScale,
-	const FSolverVec3 LocalSpaceLocation,
-	const FRotation3& ReferenceSpaceRotation)
+	const FSolverVec3 LocalSpaceLocation)
 {
 	if (EdgeConstraints)
 	{
@@ -2676,12 +2672,12 @@ void FClothConstraints::Update(
 			ReferenceSpaceRotation
 		);
 	}
-	if (BuoyancyField)
+	if (VelocityPressureAndBuoyancyField)
 	{
 		constexpr FSolverReal WorldScale = 100.f;
 		const bool bPointBasedWindDisablesAccurateWind = ClothingSimulationClothConsoleVariables::CVarLegacyDisablesAccurateWind.GetValueOnAnyThread();
 		const bool bEnableAerodynamics = !(bUsePointBasedWindModel && bPointBasedWindDisablesAccurateWind);
-		BuoyancyField->SetPropertiesAndBuoyancy(
+		VelocityPressureAndBuoyancyField->SetPropertiesAndBuoyancy(
 			ConfigProperties,
 			WeightMaps,
 			WorldScale,
diff --git a/Engine/Plugins/ChaosCloth/Source/ChaosCloth/Private/ChaosCloth/ChaosClothingSimulationCloth.cpp b/Engine/Plugins/ChaosCloth/Source/ChaosCloth/Private/ChaosCloth/ChaosClothingSimulationCloth.cpp
--- a/Engine/Plugins/ChaosCloth/Source/ChaosCloth/Private/ChaosCloth/ChaosClothingSimulationCloth.cpp
+++ b/Engine/Plugins/ChaosCloth/Source/ChaosCloth/Private/ChaosCloth/ChaosClothingSimulationCloth.cpp
@@ -468,7 +468,7 @@ void FClothingSimulationCloth::FLODData::Update(FClothingSimulationSolver* Solve
 		const FReal LocalSpaceScaleInv = 1. / LocalSpaceScale;
 		const Softs::FSolverReal MeshScale = Cloth->Mesh->GetScale() * LocalSpaceScaleInv;
 		const Softs::FSolverReal MaxDistancesScale = (Softs::FSolverReal)Cloth->MaxDistancesMultiplier;
-		ClothConstraints.Update(Cloth->Config->GetProperties(SolverDatum.LODIndex), WeightMaps, VertexSets, FaceSets, FaceIntMaps, MeshScale, MaxDistancesScale, Solver->GetLocalSpaceRotation(), LocalSpaceScale, Solver->GetLocalSpaceLocation(), Cloth->ReferenceSpaceTransform.GetRotation());
+		ClothConstraints.Update(Cloth->Config->GetProperties(SolverDatum.LODIndex), WeightMaps, VertexSets, FaceSets, FaceIntMaps, MeshScale, MaxDistancesScale, Cloth->ReferenceSpaceTransform.GetRotation(), Solver->GetLocalSpaceRotation(), LocalSpaceScale, Solver->GetLocalSpaceLocation());
 
 		if (const TSharedPtr<Softs::FPBDSkinnedTriangleMeshCollisionConstraints>& SkinnedTriangleCollisionConstraint = ClothConstraints.GetSkinnedTriangleCollisionsConstraint())
 		{
diff --git a/Engine/Plugins/ChaosCloth/Source/ChaosCloth/Private/ChaosCloth/ChaosClothingSimulationSolver.cpp b/Engine/Plugins/ChaosCloth/Source/ChaosCloth/Private/ChaosCloth/ChaosClothingSimulationSolver.cpp
--- a/Engine/Plugins/ChaosCloth/Source/ChaosCloth/Private/ChaosCloth/ChaosClothingSimulationSolver.cpp
+++ b/Engine/Plugins/ChaosCloth/Source/ChaosCloth/Private/ChaosCloth/ChaosClothingSimulationSolver.cpp
@@ -2517,16 +2517,22 @@ FBoxSphereBounds FClothingSimulationSolver::CalculateBounds() const
 }
 
 
-void FClothingSimulationSolver::PreSimulate_BuoyancyField_GameThread(float dt, const UWorld* World, const FBox BoundingBox)
+bool FClothingSimulationSolver::UsesBuoyancy()
 {
-	auto& AllConstraints = GetAllClothConstraints();
+	static IConsoleVariable* const UsesBuoyancy = IConsoleManager::Get().FindConsoleVariable(TEXT("p.Chaos.Cloth.EnableBuoyancy"));
+	return UsesBuoyancy ? UsesBuoyancy->GetBool() && !IsLegacySolver() : false;
+}
+
+void FClothingSimulationSolver::PreSimulateBuoyancyField_GameThread(float dt, const UWorld* World, const FBox BoundingBox)
+{
+	auto& AllConstraints = ClothsConstraints;
 
 	for (auto& Pair : AllConstraints)
 	{
 		int32 Key = Pair.Key;
 		::Chaos::FClothConstraints* Constraints = Pair.Value.Get();
 
-		if (FBuoyancyField *BuoyancyField = Constraints->GetBuoyancyField().Get())
+		if (Buoyancy::FBuoyancyField *BuoyancyField = Constraints->GetBuoyancyField().Get())
 		{
 			BuoyancyField->PreSimulate_GameThread(dt, World, BoundingBox);
 		}
diff --git a/Engine/Plugins/ChaosCloth/Source/ChaosCloth/Public/ChaosCloth/ChaosClothConstraints.h b/Engine/Plugins/ChaosCloth/Source/ChaosCloth/Public/ChaosCloth/ChaosClothConstraints.h
--- a/Engine/Plugins/ChaosCloth/Source/ChaosCloth/Public/ChaosCloth/ChaosClothConstraints.h
+++ b/Engine/Plugins/ChaosCloth/Source/ChaosCloth/Public/ChaosCloth/ChaosClothConstraints.h
@@ -12,7 +12,10 @@
 #include "Chaos/Deformable/GaussSeidelCorotatedCodimensionalConstraints.h"
 #include "Chaos/PBDExtremeDeformationConstraints.h"
 
-class FBuoyancyField;
+namespace Buoyancy
+{
+	class FBuoyancyField;
+}
 
 struct FManagedArrayCollection;
 
@@ -93,9 +96,9 @@ namespace Chaos
 			Softs::FSolverReal MeshScale,
 			Softs::FSolverReal MaxDistancesScale = (Softs::FSolverReal)1.,
 			const FRotation3& LocalSpaceRotation = FRotation3::Identity,
+			const FRotation3& ReferenceSpaceRotation = FRotation3::Identity,
 			const FReal LocalSpaceScale = (Softs::FSolverReal)1.,
-			const FSolverVec3 LocalSpaceLocation = FSolverVec3::ZeroVector,
-			const FRotation3& ReferenceSpaceRotation = FRotation3::Identity);
+			const FSolverVec3 LocalSpaceLocation = FSolverVec3::ZeroVector);
 
 		// NOTE: this only does something if using the PBDSolver. Force-based solver constraints are activated automatically
 		// when activating a particle range.
@@ -131,7 +134,7 @@ namespace Chaos
 		{ return SelfCollisionSphereConstraints; }
 		const TSharedPtr<Softs::FVelocityAndPressureField>& GetVelocityAndPressureField() const { return VelocityAndPressureField; }
 		
-		const TSharedPtr<FBuoyancyField>& GetBuoyancyField() const { return BuoyancyField; }
+		const TSharedPtr<Buoyancy::FBuoyancyField>& GetBuoyancyField() const { return VelocityPressureAndBuoyancyField; }
 
 		const TSharedPtr<Softs::FExternalForces>& GetExternalForces() const { return ExternalForces; }
 		const TSharedPtr<Softs::FPBDSoftBodyCollisionConstraint>& GetCollisionConstraint() const { return CollisionConstraint; }
@@ -196,7 +199,7 @@ namespace Chaos
 			const TMap<FString, TConstArrayView<FRealSingle>>& WeightMaps,
 			const FTriangleMesh& TriangleMesh
 		);
-		void CreateBuoyancyField(
+		void CreateVelocityPressureAndBuoyancyField(
 			const Softs::FCollectionPropertyConstFacade& ConfigProperties,
 			const TMap<FString, TConstArrayView<FRealSingle>>& WeightMaps,
 			const FTriangleMesh& TriangleMesh
@@ -252,7 +255,7 @@ namespace Chaos
 		//~ Begin Force-based solver only constraints
 		Softs::FSolverVec3 SolverWindVelocity; // Set from solver and added to wind from the config
 		TSharedPtr<Softs::FVelocityAndPressureField> VelocityAndPressureField;
-		TSharedPtr<FBuoyancyField> BuoyancyField;
+		TSharedPtr<Buoyancy::FBuoyancyField> VelocityPressureAndBuoyancyField;
 		TSharedPtr<Softs::FExternalForces> ExternalForces;
 		TSharedPtr<Softs::FPBDSoftBodyCollisionConstraint> CollisionConstraint;
 		TSharedPtr<Softs::FPBDSkinnedTriangleMeshCollisionConstraints> SkinnedTriangleCollisionsConstraint;
diff --git a/Engine/Plugins/ChaosCloth/Source/ChaosCloth/Public/ChaosCloth/ChaosClothingSimulationSolver.h b/Engine/Plugins/ChaosCloth/Source/ChaosCloth/Public/ChaosCloth/ChaosClothingSimulationSolver.h
--- a/Engine/Plugins/ChaosCloth/Source/ChaosCloth/Public/ChaosCloth/ChaosClothingSimulationSolver.h
+++ b/Engine/Plugins/ChaosCloth/Source/ChaosCloth/Public/ChaosCloth/ChaosClothingSimulationSolver.h
@@ -414,10 +414,10 @@ namespace Chaos
 		}
 
 		const FClothConstraints& GetClothConstraints(int32 ParticleRangeId) const { return *ClothsConstraints.FindChecked(ParticleRangeId); }
-		FClothConstraints& GetClothConstraints(int32 ParticleRangeId) { return *ClothsConstraints.FindChecked(ParticleRangeId); }
-		const TMap<int32, TUniquePtr<FClothConstraints>> & GetAllClothConstraints() const { return ClothsConstraints; }
+		FClothConstraints& GetClothConstraints(int32 ParticleRangeId) { return *ClothsConstraints.FindChecked(ParticleRangeId); }		
 
-		CHAOSCLOTH_API void PreSimulate_BuoyancyField_GameThread(float dt, const UWorld* World, const FBox BoundingBox);		
+		CHAOSCLOTH_API bool UsesBuoyancy();
+		CHAOSCLOTH_API void PreSimulateBuoyancyField_GameThread(float dt, const UWorld* World, const FBox BoundingBox);		
 
 		CHAOSCLOTH_API uint32 GetNumParticles() const;
 		CHAOSCLOTH_API int32 GetNumActiveParticles() const;
diff --git a/Engine/Plugins/ChaosClothAsset/Source/ChaosClothAssetEngine/Private/ChaosClothAsset/ClothSimulationProxy.cpp b/Engine/Plugins/ChaosClothAsset/Source/ChaosClothAssetEngine/Private/ChaosClothAsset/ClothSimulationProxy.cpp
--- a/Engine/Plugins/ChaosClothAsset/Source/ChaosClothAssetEngine/Private/ChaosClothAsset/ClothSimulationProxy.cpp
+++ b/Engine/Plugins/ChaosClothAsset/Source/ChaosClothAssetEngine/Private/ChaosClothAsset/ClothSimulationProxy.cpp
@@ -461,9 +461,8 @@ namespace UE::Chaos::ClothAsset
 
 		// Update world forces
 		if (UWorld* const World = ClothComponentAdapter.GetOwnerComponent().GetWorld())
-		{
-			static IConsoleVariable* const UsesBuoyancy = IConsoleManager::Get().FindConsoleVariable(TEXT("p.Chaos.Cloth.EnableBuoyancy"));
-			const bool bUsesBuoyancy = UsesBuoyancy ? UsesBuoyancy->GetBool() : false;
+		{			
+			const bool bUsesBuoyancy = Solver->UsesBuoyancy();
 
 			FBox BoundingBox;
 			if (World->PhysicsField || bUsesBuoyancy)
@@ -481,7 +480,7 @@ namespace UE::Chaos::ClothAsset
 			if (bUsesBuoyancy)
 			{				
 				// find all buoyancy fields and compute collisions				
-				Solver->PreSimulate_BuoyancyField_GameThread(DeltaTime, World, BoundingBox);
+				Solver->PreSimulateBuoyancyField_GameThread(DeltaTime, World, BoundingBox);
 			}
 		}
 
diff --git a/Engine/Plugins/ChaosClothAssetDataflowNodes/Source/ChaosClothAssetDataflowNodes/Private/ChaosClothAsset/SimulationAerodynamicsConfigNode.cpp b/Engine/Plugins/ChaosClothAssetDataflowNodes/Source/ChaosClothAssetDataflowNodes/Private/ChaosClothAsset/SimulationAerodynamicsConfigNode.cpp
--- a/Engine/Plugins/ChaosClothAssetDataflowNodes/Source/ChaosClothAssetDataflowNodes/Private/ChaosClothAsset/SimulationAerodynamicsConfigNode.cpp
+++ b/Engine/Plugins/ChaosClothAssetDataflowNodes/Source/ChaosClothAssetDataflowNodes/Private/ChaosClothAsset/SimulationAerodynamicsConfigNode.cpp
@@ -63,7 +63,7 @@ void FChaosClothAssetSimulationAerodynamicsConfigNode::AddProperties(FPropertyHe
 		
 	if (bUseWithWaterBodies)
 	{
-		PropertyHelper.SetProperty(this, &BuoyancyThickness);
+		PropertyHelper.SetProperty(this, &ClothDensityInWater);
 		PropertyHelper.SetProperty(this, &WaterDensity);
 		PropertyHelper.SetProperty(this, &WaterTurbulenceRatio);
 	}
diff --git a/Engine/Plugins/ChaosClothAssetDataflowNodes/Source/ChaosClothAssetDataflowNodes/Public/ChaosClothAsset/SimulationAerodynamicsConfigNode.h b/Engine/Plugins/ChaosClothAssetDataflowNodes/Source/ChaosClothAssetDataflowNodes/Public/ChaosClothAsset/SimulationAerodynamicsConfigNode.h
--- a/Engine/Plugins/ChaosClothAssetDataflowNodes/Source/ChaosClothAssetDataflowNodes/Public/ChaosClothAsset/SimulationAerodynamicsConfigNode.h
+++ b/Engine/Plugins/ChaosClothAssetDataflowNodes/Source/ChaosClothAssetDataflowNodes/Public/ChaosClothAsset/SimulationAerodynamicsConfigNode.h
@@ -100,19 +100,27 @@ public:
 	bool bUseWithWaterBodies = false;
 	
 	/* 
-	 * Thickness for cloth when used for buoyancy calculations
+	 * The density of the cloth for buoyancy calculations when submerged in water.
+	 * Given in kg/m^3
+	 * If you want bodies to be neutrally buoyancy, make this the water density.  Larger values make it sink, smaller values make it
+	 * rise to the surface.
 	 */
-	UPROPERTY(EditAnywhere, Category = "Aerodynamics", Meta = (UIMin = "0", UIMax = "1", InteractorName = "BuoyancyThickness", EditCondition = "bUseWithWaterBodies"))
-	float BuoyancyThickness = .00001;
+	UPROPERTY(EditAnywhere, Category = "Aerodynamics", Meta = (UIMin = "0", UIMax = "10000", InteractorName = "ClothDensityInWater", EditCondition = "bUseWithWaterBodies"))
+	float ClothDensityInWater = 1000.;
 
-	/* 
-	 * Density of fluid underwater
+	/**
+	 * The density of the water where buoyancy is taking effect
+	 * The fluid density is given in kg/m^3.
+	 * Water density is 1000 kg/m^3	 
 	 */
 	UPROPERTY(EditAnywhere, Category = "Aerodynamics", Meta = (UIMin = "0", UIMax = "10", ClampMin = "0", ClampMax = "10000", InteractorName = "WaterDensity", EditCondition = "bUseWithWaterBodies"))
-	float WaterDensity = 1000.0f;
+	float WaterDensity = 1000.;
 
-	/* 
-	 * Turbulence of fluid underwater
+	/**
+	 * Ratio of aerodynamic forces applied as turbulent rather than laminar flow.
+	 * Typically, you will want to use turbulent flow for wind and laminar flow for swimming underwater (unless the water is flowing very rapidly).
+	 * Turbulent forces scale like velocity squared. Laminar forces scale like velocity.
+	 * A value of 0 is appropriate for water.
 	 */
 	UPROPERTY(EditAnywhere, Category = "Aerodynamics", Meta = (ClampMin = "0", ClampMax = "1", InteractorName = "WaterTurbulenceRatio", EditCondition = "bUseWithWaterBodies"))
 	float WaterTurbulenceRatio = 0.f;
diff --git a/Engine/Plugins/Experimental/Buoyancy/Source/Runtime/Private/BuoyancyField.cpp b/Engine/Plugins/Experimental/Buoyancy/Source/Runtime/Private/BuoyancyField.cpp
--- a/Engine/Plugins/Experimental/Buoyancy/Source/Runtime/Private/BuoyancyField.cpp
+++ b/Engine/Plugins/Experimental/Buoyancy/Source/Runtime/Private/BuoyancyField.cpp
@@ -13,20 +13,22 @@ using namespace Chaos::Softs;
 
 extern bool bBuoyancyDebugDraw;
 
+namespace Buoyancy
+{
 
-
-// #todo(dmp): resolve threading issues so this can work with Chaos Cloth.  See notes below.  For now, this
-// will be left unused and instead we'll only do thread safe GT queries
-bool bComputeWaterPerTri = false;
-// FAutoConsoleVariableRef CVarBuoyancyAlgorithmsFieldComputeWaterPerTri(TEXT("p.Buoyancy.Algorithms.Field.ComputeWaterPerTri"), bComputeWaterPerTri, TEXT(""));
-
-bool bUseAccurateSubmersion = false;
-FAutoConsoleVariableRef CVarBuoyancyAlgorithmsFieldUseAccurateSubmersion(TEXT("p.Buoyancy.Algorithms.Field.UseAccurateSubmersion"), bUseAccurateSubmersion, TEXT(""));
-
-namespace BuoyancyPrivate
+namespace Private
 {
 	static float VelocityFieldMaxVelocity = 0.f;
 	static FAutoConsoleVariableRef CVarChaosVelocityFieldMaxVelocity(TEXT("p.Buoyancy.BuoyancyField.MaxVelocity"), VelocityFieldMaxVelocity, TEXT("The maximum relative velocity to process the aerodynamics forces with."));
+
+	// #todo(dmp): resolve threading issues so this can work with Chaos Cloth.  See notes below.  For now, this
+	// will be left unused and instead we'll only do thread safe GT queries
+	bool bComputeWaterPerTri = false;
+	// FAutoConsoleVariableRef CVarBuoyancyAlgorithmsFieldComputeWaterPerTri(TEXT("p.Buoyancy.Algorithms.Field.ComputeWaterPerTri"), bComputeWaterPerTri, TEXT(""));
+
+	bool bUseAccurateSubmersion = false;
+	FAutoConsoleVariableRef CVarBuoyancyAlgorithmsFieldUseAccurateSubmersion(TEXT("p.Buoyancy.Algorithms.Field.UseAccurateSubmersion"), bUseAccurateSubmersion, TEXT(""));
+
 }
 
 void FBuoyancyField::SetPropertiesAndBuoyancy(
@@ -44,14 +46,19 @@ void FBuoyancyField::SetPropertiesAndBuoyancy(
 	
 	// if none are set, then we disable buoyancy
 	bUseWaterBodies = 
-		BuoyancyThicknessIndex != INDEX_NONE && 
+		ClothDensityInWaterIndex != INDEX_NONE && 
 		WaterDensityIndex != INDEX_NONE && 
 		WaterTurbulenceRatioIndex != INDEX_NONE;
 
-	BuoyancyTriThickness = BuoyancyThicknessIndex != INDEX_NONE ? float(GetBuoyancyThickness(PropertyCollection)) : 0.1f;	
+	ClothDensityInWater = ClothDensityInWaterIndex != INDEX_NONE ? float(GetClothDensityInWater(PropertyCollection)) : 0.1f;	
+	
+	// #todo(dmp): note the Buoyancy plugin uses g/cm^3 for water density, so this would ideally be unified.
 	WaterDensity = WaterDensityIndex != INDEX_NONE ? float(GetWaterDensity(PropertyCollection)) : 1000.f;	
+			
 	WaterTurbulenceRatio = WaterTurbulenceRatioIndex != INDEX_NONE ? float(GetWaterTurbulenceRatio(PropertyCollection)) : 0.f;	
 
+	// convert to kg/cm^3
+	ClothRho = FMath::Max(ClothDensityInWater / FMath::Cube(WorldScale), (FSolverReal)0.);
 	WaterRho = FMath::Max(WaterDensity / FMath::Cube(WorldScale), (FSolverReal)0.);
 	
 	constexpr FSolverReal OneQuarter = (FSolverReal)0.25;
@@ -61,7 +68,7 @@ void FBuoyancyField::SetPropertiesAndBuoyancy(
 	LocalSpaceLocation = InLocalSpaceLocation;
 }
 
-bool FBuoyancyField::PreSimulate_GameThread(float DeltaTime, const UWorld* World, const FBox BoundingBox)
+void FBuoyancyField::PreSimulate_GameThread(float DeltaTime, const UWorld* World, const FBox BoundingBox)
 {
 	check(IsInGameThread());
 	
@@ -77,7 +84,7 @@ bool FBuoyancyField::PreSimulate_GameThread(float DeltaTime, const UWorld* World
 		FVector WaterPlaneNormal = {0,0,0};
 		FVector WaterPlaneVelocity = {0,0,0};
 
-		bool ComputePlane = !bComputeWaterPerTri;
+		bool ComputePlane = !Private::bComputeWaterPerTri;
 
 		bool HasCollision = BuoyancySubsystem->FindOverlappingWaterBodies(BoundingBox, ComputePlane,
 			WaterPlaneLocation,
@@ -92,8 +99,6 @@ bool FBuoyancyField::PreSimulate_GameThread(float DeltaTime, const UWorld* World
 			WaterPlaneNormal,
 			WaterPlaneVelocity);
 	}
-
-	return true;
 }
 
 void FBuoyancyField::Apply(FSolverParticlesRange& InParticles, const FSolverReal Dt) const
@@ -107,7 +112,7 @@ void FBuoyancyField::Apply(FSolverParticlesRange& InParticles, const FSolverReal
 	}
 	else
 	{
-		const FSolverReal MaxVelocitySquared = (BuoyancyPrivate::VelocityFieldMaxVelocity > 0.f) ? FMath::Square((FSolverReal)BuoyancyPrivate::VelocityFieldMaxVelocity) : TNumericLimits<FSolverReal>::Max();
+		const FSolverReal MaxVelocitySquared = (Private::VelocityFieldMaxVelocity > 0.f) ? FMath::Square((FSolverReal)Private::VelocityFieldMaxVelocity) : TNumericLimits<FSolverReal>::Max();
 
 		// update according to water bodies
 		const TConstArrayView<FSolverVec3>& Xs = InParticles.XArray();
@@ -120,7 +125,7 @@ void FBuoyancyField::Apply(FSolverParticlesRange& InParticles, const FSolverReal
 
 		TSharedPtr<FBuoyancyWaterSplineData> WaterBodySplineData = nullptr;
 
-		if (!bComputeWaterPerTri)
+		if (!Private::bComputeWaterPerTri)
 		{
 			TRACE_CPUPROFILER_EVENT_SCOPE(FBuoyancyField_Apply_ComputeWaterPerCloth);
 			const FVector WorldAvgX = WaterBodyCollisions.WorldBounds.GetCenter();
@@ -140,6 +145,47 @@ void FBuoyancyField::Apply(FSolverParticlesRange& InParticles, const FSolverReal
 			return;
 		}
 
+		// loop over particles and compute buoyancy forces
+		for (int32 Index = 0; Index < InParticles.GetRangeSize(); ++Index)
+		{
+			if (InParticles.InvM(Index) != (FSolverReal) 0. && ClothRho > (FSolverReal) 0.)
+			{				
+				const FSolverVec3 WorldPosition = InParticles.GetP(Index) * LocalSpaceScale + LocalSpaceLocation;
+			
+				// optionally compute water properties at tri center			
+				if (Private::bComputeWaterPerTri)
+				{				
+					TRACE_CPUPROFILER_EVENT_SCOPE(FBuoyancyField_Apply_ComputeWaterPerTri);
+								
+					const FVector WorldPositionFVector = FVector(WorldPosition);
+					WaterBodyCollisions.TheBuoyancySubsystem->QueryWaterBody(WorldPositionFVector, WaterBodySplineData, WaterVelocity, WaterN, WaterPlanePos);
+				}
+
+				const FSolverVec3 WaterNSolverVec = FSolverVec3(WaterN);
+
+				if ((WorldPosition - WaterPlanePos).Dot(WaterNSolverVec) < SMALL_NUMBER)
+				{
+					const FSolverVec3 TriGravity = ExternalForces->GetScaledGravity(Index) * LocalSpaceScale;
+									
+					// the fraction of the current body (ie: particle) is submerged.  We
+					// treat particles as either fully submerged or not.  It could be
+					// TriSubmersionArea / TriArea for instance if we did this on triangles
+					const FSolverReal SubmersionFraction = 1.;
+				
+					// fraction of the particle
+					// volume = mass / density - scale mass by submersion fraction then divide by rho
+					const FSolverReal SubmergedTriVolume = (InParticles.M(Index) * SubmersionFraction) / ClothRho;
+
+					// Archimedes Principle						      cm/sec^2  * kg/cm^3  *     cm^3
+					const FSolverVec3 Force = ((FSolverReal) -1.0) * TriGravity * WaterRho * SubmergedTriVolume;
+
+					// Update acceleration
+					InParticles.Acceleration(Index) += InParticles.InvM(Index) * Force;
+				}
+			}
+		}
+
+		// loop over triangles and compute drag forces
 		for (int32 ElementIndex = 0; ElementIndex < Elements.Num(); ++ElementIndex)
 		{
 			const TVec3<int32>& Element = Elements[ElementIndex];
@@ -162,7 +208,7 @@ void FBuoyancyField::Apply(FSolverParticlesRange& InParticles, const FSolverReal
 			const FVector WorldTriCenter = FVector((WorldTriPositions[0] + WorldTriPositions[1] + WorldTriPositions[2]) / 3.f);
 
 			// optionally compute water properties at tri center			
-			if (bComputeWaterPerTri)
+			if (Private::bComputeWaterPerTri)
 			{				
 				TRACE_CPUPROFILER_EVENT_SCOPE(FBuoyancyField_Apply_ComputeWaterPerTri);
 												
@@ -210,7 +256,7 @@ void FBuoyancyField::Apply(FSolverParticlesRange& InParticles, const FSolverReal
 				Xs[Element[1]] - Xs[Element[0]]);
 
 			const FSolverReal NormalLength = TriNormal.SafeNormalize();
-			
+			const FSolverReal TriArea = (FSolverReal) .5 * NormalLength;
 			FSolverReal SubmergedTriArea;
 
 			// In simple submersion, any triangle that is submerged is considered fully submerged
@@ -218,7 +264,7 @@ void FBuoyancyField::Apply(FSolverParticlesRange& InParticles, const FSolverReal
 			{
 				continue;
 			}
-			else if (!bUseAccurateSubmersion || TriIsFullySubmerged)
+			else if (!Private::bUseAccurateSubmersion || TriIsFullySubmerged)
 			{
 				// #todo(dmp): simple submersion could be improved by computing the circle circumscribed by the tri then
 				// doing analytical submersion depth based on sphere/plane penetration depth
@@ -229,7 +275,7 @@ void FBuoyancyField::Apply(FSolverParticlesRange& InParticles, const FSolverReal
 				// lots of ways to do this, but plenty of errors due to stretched out triangles misrepresenting the tri area
 				// and it is unclear if it'd be faster than the accurate geometric method below
 				
-				SubmergedTriArea = (FSolverReal) .5 * NormalLength;
+				SubmergedTriArea = TriArea;
 			}
 			else
 			{
@@ -296,11 +342,6 @@ void FBuoyancyField::Apply(FSolverParticlesRange& InParticles, const FSolverReal
 				Vs[Element[2]]);
 			
 			Force = CalculateForce(TriNormal, TriSurfaceVelocity, SubmergedTriArea, WaterVelocityOnPlane, CdI, CdO, ClI, ClO, Cp, MaxVelocitySquared);
-									
-			// buoyancy force		
-			const FSolverReal SubmergedTriVolume = SubmergedTriArea * BuoyancyTriThickness;			
-			//                               cm/sec^2     *    kg/cm^3          *     cm^3
-		    Force += ((FSolverReal) -1.0) * SolverGravity * WaterRho * SubmergedTriVolume;			
 
 			// increment acceleration on each vertex
 			// #todo(dmp): only update vertices that are submerged or whole triangle?
@@ -323,4 +364,5 @@ void FBuoyancyField::Apply(FSolverParticlesRange& InParticles, const FSolverReal
 #endif				
 		}
 	}
+}
 }
\ No newline at end of file
diff --git a/Engine/Plugins/Experimental/Buoyancy/Source/Runtime/Public/BuoyancyField.h b/Engine/Plugins/Experimental/Buoyancy/Source/Runtime/Public/BuoyancyField.h
--- a/Engine/Plugins/Experimental/Buoyancy/Source/Runtime/Public/BuoyancyField.h
+++ b/Engine/Plugins/Experimental/Buoyancy/Source/Runtime/Public/BuoyancyField.h
@@ -2,11 +2,14 @@
 #pragma once
 
 #include "Chaos/VelocityField.h"
+#include "Chaos/SoftsExternalForces.h"
 #include "BuoyancySubsystem.h"
 
 using Chaos::Softs::FSolverVec3;
 using Chaos::Softs::FSolverReal;
 
+namespace Buoyancy
+{
 // Velocity field used solely for aerodynamics effects including water submersion
 class FBuoyancyField : public Chaos::Softs::FVelocityAndPressureField
 {
@@ -14,7 +17,7 @@ public:
 	
 	explicit FBuoyancyField(const Chaos::Softs::FCollectionPropertyConstFacade& PropertyCollection) : 
 		Chaos::Softs::FVelocityAndPressureField(PropertyCollection),
-		BuoyancyThicknessIndex(PropertyCollection),
+		ClothDensityInWaterIndex(PropertyCollection),
 		WaterDensityIndex(PropertyCollection),
 		WaterTurbulenceRatioIndex(PropertyCollection)
 	{
@@ -25,18 +28,19 @@ public:
 		const Chaos::FTriangleMesh* TriangleMesh,
 		const Chaos::Softs::FCollectionPropertyConstFacade& PropertyCollection,
 		const TMap<FString, TConstArrayView<Chaos::FRealSingle>>& Weightmaps,
-		Chaos::FSolverReal WorldScale)
+		Chaos::FSolverReal WorldScale, TSharedPtr<Chaos::Softs::FExternalForces> InExternalForces)
 		: Chaos::Softs::FVelocityAndPressureField(Particles, TriangleMesh, PropertyCollection, Weightmaps, WorldScale),
-		BuoyancyThicknessIndex(PropertyCollection),
+		ClothDensityInWaterIndex(PropertyCollection),
 		WaterDensityIndex(PropertyCollection),
-		WaterTurbulenceRatioIndex(PropertyCollection)
+		WaterTurbulenceRatioIndex(PropertyCollection),
+		ExternalForces(InExternalForces)
 		
 	{		
 	}
 
 	// Construct an uninitialized field. Mesh, properties, and velocity will have to be set for this field to be valid.
 	FBuoyancyField() : Chaos::Softs::FVelocityAndPressureField(),
-		BuoyancyThicknessIndex(ForceInit),
+		ClothDensityInWaterIndex(ForceInit),
 		WaterDensityIndex(ForceInit),
 		WaterTurbulenceRatioIndex(ForceInit)
 	{
@@ -58,15 +62,13 @@ public:
 		const Chaos::FRotation3& ReferenceSpaceRotation = Chaos::FRotation3::Identity
 		);
 
-	BUOYANCY_API bool PreSimulate_GameThread(float DeltaTime, const UWorld* World, const FBox BoundingBox);
-
-	void SetSolverGravity(FSolverVec3 Gravity) { SolverGravity = Gravity; }
+	BUOYANCY_API void PreSimulate_GameThread(float DeltaTime, const UWorld* World, const FBox BoundingBox);
 
-	UE_CHAOS_DECLARE_INDEXLESS_PROPERTYCOLLECTION_NAME(BuoyancyThickness, float);
+	UE_CHAOS_DECLARE_INDEXLESS_PROPERTYCOLLECTION_NAME(ClothDensityInWater, float);
 	UE_CHAOS_DECLARE_INDEXLESS_PROPERTYCOLLECTION_NAME(WaterDensity, float);
 	UE_CHAOS_DECLARE_INDEXLESS_PROPERTYCOLLECTION_NAME(WaterTurbulenceRatio, float);
 
-protected:
+private:
 
 	FSolverVec3 CalculateForce(FSolverVec3 N, const FSolverVec3& SurfaceVelocity, const FSolverReal SubmergedTriArea, const FSolverVec3& FluidVelocity, const FSolverReal CdI, const FSolverReal CdO, const FSolverReal ClI, const FSolverReal ClO, const FSolverReal Cp, const FSolverReal MaxVelocitySquared) const
 	{		
@@ -90,16 +92,21 @@ protected:
 
 	FBuoyancyCollisionData WaterBodyCollisions;
 	FSolverReal LocalSpaceScale = 1.f;
-	FSolverVec3 LocalSpaceLocation = {0.f,0.f,0.f};
-	FSolverVec3 SolverGravity = {0.f,0.f,0.f};
-	FSolverReal BuoyancyTriThickness = .01f;
+	FSolverVec3 LocalSpaceLocation = {0.f,0.f,0.f};	
+	FSolverReal ClothDensityInWater = 1000.f;
 	FSolverReal WaterDensity = 1.f;
 	FSolverReal WaterTurbulenceRatio = 1.f;
 
 	FSolverReal WaterRho = 1.f;
 	FSolverReal QuarterWaterRho = 1.f;
+	
+	FSolverReal ClothRho = 1.f;
 
-	UE_CHAOS_DECLARE_INDEXED_PROPERTYCOLLECTION_NAME(BuoyancyThickness, float);
+	UE_CHAOS_DECLARE_INDEXED_PROPERTYCOLLECTION_NAME(ClothDensityInWater, float);
 	UE_CHAOS_DECLARE_INDEXED_PROPERTYCOLLECTION_NAME(WaterDensity, float);
 	UE_CHAOS_DECLARE_INDEXED_PROPERTYCOLLECTION_NAME(WaterTurbulenceRatio, float);
-};
\ No newline at end of file
+
+	TSharedPtr<Chaos::Softs::FExternalForces> ExternalForces;
+};
+
+}
\ No newline at end of file
diff --git a/Engine/Source/Runtime/Experimental/Chaos/Public/Chaos/VelocityField.h b/Engine/Source/Runtime/Experimental/Chaos/Public/Chaos/VelocityField.h
--- a/Engine/Source/Runtime/Experimental/Chaos/Public/Chaos/VelocityField.h
+++ b/Engine/Source/Runtime/Experimental/Chaos/Public/Chaos/VelocityField.h
@@ -108,7 +108,7 @@ public:
 	{
 	}
 
-	virtual ~FVelocityAndPressureField() {}
+	virtual ~FVelocityAndPressureField() = default;
 
 	CHAOS_API void UpdateForces(const FSolverParticles& InParticles, const FSolverReal /*Dt*/);
 
@@ -225,8 +225,8 @@ protected:
 		const TConstArrayView<FRealSingle>& PressureMultipliers);
 
 
-	FSolverVec3 CalculateForceLowLevel(const FSolverVec3 &InV, const FSolverVec3 &InN, const FSolverReal InQuarterRho, const FSolverReal InDoubleArea, 
-		FSolverReal InCdI, FSolverReal InCdO, const FSolverReal InClI, const FSolverReal InClO, const FSolverReal InCp, const FSolverReal InTr) const
+	static FSolverVec3 CalculateForceLowLevel(const FSolverVec3 &InV, const FSolverVec3 &InN, const FSolverReal InQuarterRho, const FSolverReal InDoubleArea, 
+		FSolverReal InCdI, FSolverReal InCdO, const FSolverReal InClI, const FSolverReal InClO, const FSolverReal InCp, const FSolverReal InTr)
 	{		
 		// Set the aerodynamic forces
 		const FSolverReal VDotN = FSolverVec3::DotProduct(InV, InN);
