CHAOS_DEBUG_DRAW instead of ENABLE_DEBUG_DRAW in Buoyancy code

#Synced-CL 50915206

diff --git a/Engine/Plugins/Experimental/Buoyancy/Source/Runtime/Private/BuoyancyAlgorithms.cpp b/Engine/Plugins/Experimental/Buoyancy/Source/Runtime/Private/BuoyancyAlgorithms.cpp
--- a/Engine/Plugins/Experimental/Buoyancy/Source/Runtime/Private/BuoyancyAlgorithms.cpp
+++ b/Engine/Plugins/Experimental/Buoyancy/Source/Runtime/Private/BuoyancyAlgorithms.cpp
@@ -290,7 +290,7 @@ namespace BuoyancyAlgorithms
 			ScaleSubmergedVolume(Evolution, SubmergedParticle, true, SubmergedVol, TotalVol);
 
 
-#if ENABLE_DRAW_DEBUG
+#if CHAOS_DEBUG_DRAW
 			if (bBuoyancyDebugDraw)
 			{
 				Chaos::FDebugDrawQueue::GetInstance().DrawDebugPoint(SubmergedCoM, FColor::Yellow, false, -1.f, -1, 15.f);
@@ -364,7 +364,7 @@ namespace BuoyancyAlgorithms
 			}
 			const FAABB3 WorldBox = LocalBox.TransformedAABB(ShapeWorldTransform);
 
-#if ENABLE_DRAW_DEBUG
+#if CHAOS_DEBUG_DRAW
 			//if (bBuoyancyDebugDraw)
 			//{
 			//	Chaos::FDebugDrawQueue::GetInstance().DrawDebugBox(
@@ -415,7 +415,7 @@ namespace BuoyancyAlgorithms
 
 
 
-#if ENABLE_DRAW_DEBUG
+#if CHAOS_DEBUG_DRAW
 					if (bBuoyancyDebugDraw)
 					{
 						//Chaos::FDebugDrawQueue::GetInstance().DrawDebugBox(
@@ -765,7 +765,7 @@ namespace BuoyancyAlgorithms
 						OutSubmergedVol, OutSubmergedCoM, OutForce, OutTorque, OutBuoyancyForce, OutBuoyancyTorque);
 				}
 
-#if ENABLE_DRAW_DEBUG
+#if CHAOS_DEBUG_DRAW
 				if (bBuoyancyDebugDraw)
 				{					
 					Chaos::FDebugDrawQueue::GetInstance().DrawDebugBox(
@@ -1039,7 +1039,7 @@ namespace BuoyancyAlgorithms
 					OutForce += TotalWorldForce;
 					OutTorque += TotalWorldTorque;					
 
-#if ENABLE_DRAW_DEBUG
+#if CHAOS_DEBUG_DRAW
 					if (bBuoyancyDebugDraw)
 					{
 						FVector CurrWaterVelocityViz = CurrWaterVelocity;
@@ -1091,7 +1091,7 @@ namespace BuoyancyAlgorithms
 		OutBuoyancyForce = WorldBuoyantForce;
 		OutBuoyancyTorque = WorldBuoyantTorque;
 
-#if ENABLE_DRAW_DEBUG
+#if CHAOS_DEBUG_DRAW
 		if (bBuoyancyDebugDraw)
 		{
 			FVector WorldBoxCenter = ShapeWorldTransform.TransformPosition(ShapeQuery.GetCenter());
@@ -1253,7 +1253,7 @@ namespace BuoyancyAlgorithms
 		// area of the triangle
 		OutArea = .5 * NormalLength;
 		
-#if ENABLE_DRAW_DEBUG
+#if CHAOS_DEBUG_DRAW
 		if (DebugDraw)
 		{
 			Chaos::FDebugDrawQueue::GetInstance().DrawDebugLine(V0, V1, FColor::Orange, false, -1.f, 0, 1.f);
diff --git a/Engine/Plugins/Experimental/Buoyancy/Source/Runtime/Private/BuoyancyField.cpp b/Engine/Plugins/Experimental/Buoyancy/Source/Runtime/Private/BuoyancyField.cpp
--- a/Engine/Plugins/Experimental/Buoyancy/Source/Runtime/Private/BuoyancyField.cpp
+++ b/Engine/Plugins/Experimental/Buoyancy/Source/Runtime/Private/BuoyancyField.cpp
@@ -349,7 +349,7 @@ void FBuoyancyField::Apply(FSolverParticlesRange& InParticles, const FSolverReal
 			InParticles.Acceleration(Element[1]) += InParticles.InvM(Element[1]) * Force;
 			InParticles.Acceleration(Element[2]) += InParticles.InvM(Element[2]) * Force;
 
-#if ENABLE_DRAW_DEBUG
+#if CHAOS_DEBUG_DRAW
 			if (bBuoyancyDebugDraw)
 			{
 				const FVector WorldAvgX = WaterBodyCollisions.WorldBounds.GetCenter();
diff --git a/Engine/Plugins/Experimental/Buoyancy/Source/Runtime/Private/BuoyancySubsystem.cpp b/Engine/Plugins/Experimental/Buoyancy/Source/Runtime/Private/BuoyancySubsystem.cpp
--- a/Engine/Plugins/Experimental/Buoyancy/Source/Runtime/Private/BuoyancySubsystem.cpp
+++ b/Engine/Plugins/Experimental/Buoyancy/Source/Runtime/Private/BuoyancySubsystem.cpp
@@ -50,7 +50,7 @@ bool bBuoyancyCallbackDataParticleValidation = true;
 FAutoConsoleVariableRef CVarBuoyancyCallbackDataParticleValidation(TEXT("p.Buoyancy.CallbackData.ParticleValidation"), bBuoyancyCallbackDataParticleValidation, TEXT(""));
 
 bool bBuoyancyDebugDraw = false;
-#if ENABLE_DRAW_DEBUG
+#if CHAOS_DEBUG_DRAW
 FAutoConsoleVariableRef CVarBuoyancyDebugDraw(TEXT("p.Buoyancy.DebugDraw"), bBuoyancyDebugDraw, TEXT(""));
 #endif
 
@@ -810,7 +810,7 @@ void FBuoyancySubsystemSimCallback::OnPreSimulate_Internal()
 	}
 
 	// If we have debug draw enabled
-#if ENABLE_DRAW_DEBUG
+#if CHAOS_DEBUG_DRAW
 	if (bBuoyancyDebugDraw)
 	{
 		// If we're using spline key cache grid, debug draw all cached points
@@ -1052,7 +1052,7 @@ bool FBuoyancySubsystemSimCallback::QuerySpline(
 			}
 		}
 			
-#if ENABLE_DRAW_DEBUG
+#if CHAOS_DEBUG_DRAW
 		if (bBuoyancyDebugDraw && !bUseAccurateIntegrationForSplines)
 		{
 			// Spline Color
@@ -1236,7 +1236,7 @@ void FBuoyancySubsystemSimCallback::ProcessAccurateInteraction(Chaos::FPBDRigids
 
 		Chaos::FConstGenericParticleHandle RigidGeneric(Interaction.RigidParticle);
 
-#if ENABLE_DRAW_DEBUG
+#if CHAOS_DEBUG_DRAW
 		if (bBuoyancyDebugDraw)
 		{
 			Chaos::FDebugDrawQueue::GetInstance().DrawDebugSphere(SubmergedCoM, 20, 10, FColor::Green, false, -1.f, -1, 2.0f);
