use collision channel defined in plugin settings for cloth overlap tests

#Synced-CL 50915206

diff --git a/Engine/Plugins/Experimental/Buoyancy/Source/Runtime/Private/BuoyancySubsystem.cpp b/Engine/Plugins/Experimental/Buoyancy/Source/Runtime/Private/BuoyancySubsystem.cpp
--- a/Engine/Plugins/Experimental/Buoyancy/Source/Runtime/Private/BuoyancySubsystem.cpp
+++ b/Engine/Plugins/Experimental/Buoyancy/Source/Runtime/Private/BuoyancySubsystem.cpp
@@ -417,7 +417,7 @@ bool UBuoyancySubsystem::FindOverlappingWaterBodies(const FBox BoundingBox,
 	QueryParams.bReturnPhysicalMaterial = false;			
 
 	FCollisionObjectQueryParams ObjectQueryParams;
-	ObjectQueryParams.AddObjectTypesToQuery(ECC_WorldStatic);
+	ObjectQueryParams.AddObjectTypesToQuery(BuoyancySettings.WaterCollisionChannel);
 
 	TArray<FOverlapResult> Overlaps;
 			
