fix out of bounds array access with half edges

#Synced-CL 52108856

diff --git a/Engine/Plugins/Experimental/Buoyancy/Source/Runtime/Private/BuoyancyAlgorithms.cpp b/Engine/Plugins/Experimental/Buoyancy/Source/Runtime/Private/BuoyancyAlgorithms.cpp
--- a/Engine/Plugins/Experimental/Buoyancy/Source/Runtime/Private/BuoyancyAlgorithms.cpp
+++ b/Engine/Plugins/Experimental/Buoyancy/Source/Runtime/Private/BuoyancyAlgorithms.cpp
@@ -1382,15 +1382,23 @@ namespace BuoyancyAlgorithms
 	{		
 		// #todo(dmp): ideally the convex would have a mapping between half edges and edges that is serialized, but instead for now we'll build it here
 		const int32 NumEdges = Convex->NumEdges();
-		HalfEdgeToEdge.AddUninitialized(NumEdges * 2);
+		HalfEdgeToEdge.AddUninitialized(Convex->GetStructureData().NumHalfEdges());
 
 		for (int32 CurrEdgeIdx = 0; CurrEdgeIdx < NumEdges; ++CurrEdgeIdx)
 		{
 			int32 HalfEdge0, HalfEdge1;
 			Convex->GetHalfEdges(CurrEdgeIdx, HalfEdge0, HalfEdge1);
 
-			HalfEdgeToEdge[HalfEdge0] = CurrEdgeIdx;
-			HalfEdgeToEdge[HalfEdge1] = CurrEdgeIdx;
+			// only process valid half edges
+			if (HalfEdge0 != INDEX_NONE)
+			{
+				HalfEdgeToEdge[HalfEdge0] = CurrEdgeIdx;
+			}
+
+			if (HalfEdge1 != INDEX_NONE)
+			{
+				HalfEdgeToEdge[HalfEdge1] = CurrEdgeIdx;
+			}
 		}
 	}
 
