- buoyancy computation optimiazations
- lift parameter exposed
- fixed bug with scaling on drag
- readjusted convex topology limits

#Synced-CL 52108856

diff --git a/Engine/Plugins/Experimental/Buoyancy/Source/Runtime/Private/BuoyancyAlgorithms.cpp b/Engine/Plugins/Experimental/Buoyancy/Source/Runtime/Private/BuoyancyAlgorithms.cpp
--- a/Engine/Plugins/Experimental/Buoyancy/Source/Runtime/Private/BuoyancyAlgorithms.cpp
+++ b/Engine/Plugins/Experimental/Buoyancy/Source/Runtime/Private/BuoyancyAlgorithms.cpp
@@ -15,6 +15,12 @@
 #include "Chaos/Sphere.h"
 #include "Chaos/ImplicitObjectScaled.h"
 
+// Set to 1 to enable per-face/per-triangle SCOPE_CYCLE_COUNTERs in the inner loops.
+// These are useful for profiling but add measurable overhead at ~1M+ iterations/tick.
+#ifndef BUOYANCY_INNER_LOOP_STATS
+#define BUOYANCY_INNER_LOOP_STATS 0
+#endif
+
 //
 // CVars
 //
@@ -280,6 +286,27 @@ namespace BuoyancyAlgorithms
 		}
 	}
 
+	// ScaleSubmergedVolume overload using precomputed particle/shape volumes
+	// to avoid repeated VisitLeafObjects traversals per shape and per particle.
+	void ScaleSubmergedVolume(const FRealSingle ParticleVol, const FRealSingle ShapeVol,
+		FRealSingle& SubmergedVol, FRealSingle& TotalVol)
+	{
+		TotalVol = ParticleVol;
+
+		if (SubmergedVol - ShapeVol > UE_SMALL_NUMBER)
+		{
+			SubmergedVol = ShapeVol;
+		}
+
+		if (ParticleVol > UE_SMALL_NUMBER &&
+			ShapeVol > UE_SMALL_NUMBER &&
+			(bBuoyancyAlgorithmsAllowVolRatioOverOne || ParticleVol < ShapeVol))
+		{
+			const float VolRatio = ParticleVol / ShapeVol;
+			SubmergedVol *= VolRatio;
+		}
+	}
+
 
 	bool ComputeSubmergedVolume(FBuoyancyParticleData& ParticleData, const FPBDRigidsEvolutionGBF& Evolution, const FGeometryParticleHandle* SubmergedParticle, const FGeometryParticleHandle* WaterParticle, const FVector& WaterX, const FVector& WaterN, int32 NumSubdivisions, float MinVolume, float& SubmergedVol, FVec3& SubmergedCoM, float& TotalVol)
 	{
@@ -638,7 +665,7 @@ namespace BuoyancyAlgorithms
 	void ComputeSubmergedVolumeAndForcesForParticle(FBuoyancyParticleData& ParticleData, 
 		const Chaos::FGeometryParticleHandle* SubmergedParticle, const FGeometryParticleHandle* WaterParticle, 
 		TSharedPtr<SamplerType> WaterSampler,
-		const Chaos::FPBDRigidsEvolution& Evolution, const float DeltaSeconds, const float WaterDensity, const float WaterDrag,
+		const Chaos::FPBDRigidsEvolution& Evolution, const float DeltaSeconds, const float WaterDensity, const float WaterDrag, const float WaterLift,
 		float &OutTotalParticleVol, float &OutTotalSubmergedVol, Chaos::FVec3 &OutTotalSubmergedCoM, Chaos::FVec3 &OutTotalForce, Chaos::FVec3 &OutTotalTorque)
 	{		
 		SCOPE_CYCLE_COUNTER(STAT_BuoyancyAlgorithms_ComputeSubmergedVolumeAndForcesForParticle)
@@ -669,9 +696,16 @@ namespace BuoyancyAlgorithms
 
 		WaterSampler->InitializeForSubmersion();
 
+		// Precompute particle and shape volumes — both are constant per particle per tick
+		// and would otherwise be recomputed via VisitLeafObjects on every shape + once per particle.
+		const FRealSingle ParticleVol = ComputeParticleVolume(Evolution, SubmergedParticle);
+		const FRealSingle ShapeVolBBox = ComputeShapeVolume(SubmergedParticle, true);
+		const FRealSingle ShapeVolGeom = ComputeShapeVolume(SubmergedParticle, false);
+
 		// Traverse the submerged particle's leaves
 		RootImplicit->VisitLeafObjects(
-			[&Evolution, DeltaSeconds, WaterDensity, WaterDrag, SubmergedParticle, &WaterSampler, 
+			[&Evolution, DeltaSeconds, WaterDensity, WaterDrag, WaterLift, SubmergedParticle, &WaterSampler,
+			ParticleVol, ShapeVolGeom,
 			ParticleIndex, &ShapeInstances, WaterShapeType, WaterShapeInstance, &ParticleWorldTransform,
 			&SubmergedShapes, &OutTotalSubmergedVol, &OutTotalSubmergedCoM,
 			&OutTotalForce, &OutTotalTorque, &OutTotalBuoyancyForce, &OutTotalBuoyancyTorque]
@@ -750,7 +784,8 @@ namespace BuoyancyAlgorithms
 					ConvexShapeQuery.Initialize();
 
 					ComputeSubmergedVolumeAndForcesForShape<FBuoyancyConvexShape, SamplerType>(SubmergedParticle, ConvexShapeQuery,
-						Evolution, DeltaSeconds, WaterDensity, WaterDrag,
+						Evolution, DeltaSeconds, WaterDensity, WaterDrag, WaterLift,
+						ParticleVol, ShapeVolGeom,
 						ShapeWorldTransformScaled, WaterSampler,
 						OutSubmergedVol, OutSubmergedCoM, OutForce, OutTorque, OutBuoyancyForce, OutBuoyancyTorque);
 				}
@@ -760,7 +795,8 @@ namespace BuoyancyAlgorithms
 					BoxShapeQuery.Initialize();
 
 					ComputeSubmergedVolumeAndForcesForShape<FBuoyancyBoxShape, SamplerType>(SubmergedParticle, BoxShapeQuery,
-						Evolution, DeltaSeconds, WaterDensity, WaterDrag,
+						Evolution, DeltaSeconds, WaterDensity, WaterDrag, WaterLift,
+						ParticleVol, ShapeVolGeom,
 						ShapeWorldTransform, WaterSampler,
 						OutSubmergedVol, OutSubmergedCoM, OutForce, OutTorque, OutBuoyancyForce, OutBuoyancyTorque);
 				}
@@ -799,25 +835,26 @@ namespace BuoyancyAlgorithms
 		OutTotalSubmergedCoM /= OutTotalSubmergedVol;
 
 		float ScaledSubmergedVol = OutTotalSubmergedVol;
-		ScaleSubmergedVolume(Evolution, SubmergedParticle, true, ScaledSubmergedVol, OutTotalParticleVol);
+		ScaleSubmergedVolume(ParticleVol, ShapeVolBBox, ScaledSubmergedVol, OutTotalParticleVol);
 	}
 
 	// forward declaration of usage	
 	template void ComputeSubmergedVolumeAndForcesForParticle<FBuoyancyConstantSplineSampler>(FBuoyancyParticleData&,
 		const Chaos::FGeometryParticleHandle*, const Chaos::FGeometryParticleHandle*,
 		TSharedPtr<FBuoyancyConstantSplineSampler>,
-		const Chaos::FPBDRigidsEvolution&, const float, const float, const float,
+		const Chaos::FPBDRigidsEvolution&, const float, const float, const float, const float,
 		float&, float&, Chaos::FVec3&, Chaos::FVec3&, Chaos::FVec3&);
 
 	template void ComputeSubmergedVolumeAndForcesForParticle<FBuoyancyShallowWaterSampler>(FBuoyancyParticleData&,
 		const Chaos::FGeometryParticleHandle*, const Chaos::FGeometryParticleHandle*,
 		TSharedPtr<FBuoyancyShallowWaterSampler>,
-		const Chaos::FPBDRigidsEvolution&, const float, const float, const float,
+		const Chaos::FPBDRigidsEvolution&, const float, const float, const float, const float,
 		float&, float&, Chaos::FVec3&, Chaos::FVec3&, Chaos::FVec3&);
 
 	template <typename ShapeType, typename SamplerType>
 	void ComputeSubmergedVolumeAndForcesForShape(const Chaos::FGeometryParticleHandle* SubmergedParticle, const ShapeType& ShapeQuery,
-		const Chaos::FPBDRigidsEvolution& Evolution, float DeltaSeconds, const float WaterDensity, const float WaterDrag,
+		const Chaos::FPBDRigidsEvolution& Evolution, float DeltaSeconds, const float WaterDensity, const float WaterDrag, const float WaterLift,
+		const float ParticleVol, const float ShapeVolGeom,
 		const Chaos::FRigidTransform3 &ShapeWorldTransform, TSharedPtr<SamplerType> WaterSampler,
 		float& OutSubmergedVol, Chaos::FVec3& OutSubmergedCoM,
 		Chaos::FVec3& OutForce, Chaos::FVec3& OutTorque,
@@ -886,23 +923,24 @@ namespace BuoyancyAlgorithms
 		int32 InteriorRefPointCount = 0;
 		
 		// cache water values for each vertex and determine broadphase of shape interacting with the water	
-		for (int32 i = 0; i < TotalNumVertices; ++i)
 		{
-			// SCOPE_CYCLE_COUNTER(STAT_BuoyancyAlgorithms_ComputeSubmergedVolumeAndForcesForShape_ComputeWorld);
-
-			// compute world space position of vertex
-			const FVector CurrWorldVertexPosition = ShapeWorldTransform.TransformPosition(ShapeQuery.GetVertex(i));
-			WorldVertexPosition[i] = CurrWorldVertexPosition;
-
-			if ((CurrWorldVertexPosition - ShapeWaterP).Dot(ShapeWaterN) < SMALL_NUMBER)
-			{
-				VertexIsUnderwater[i] = true;
-				InteriorRefPoint += WorldVertexPosition[i];
-				InteriorRefPointCount++;				
-			}
-			else
-			{
-				VertexIsUnderwater[i] = false;
+			SCOPE_CYCLE_COUNTER(STAT_BuoyancyAlgorithms_ComputeSubmergedVolumeAndForcesForShape_ComputeWorld);
+			for (int32 i = 0; i < TotalNumVertices; ++i)
+			{			
+				// compute world space position of vertex
+				const FVector CurrWorldVertexPosition = ShapeWorldTransform.TransformPosition(ShapeQuery.GetVertex(i));
+				WorldVertexPosition[i] = CurrWorldVertexPosition;
+
+				if ((CurrWorldVertexPosition - ShapeWaterP).Dot(ShapeWaterN) < SMALL_NUMBER)
+				{
+					VertexIsUnderwater[i] = true;
+					InteriorRefPoint += WorldVertexPosition[i];
+					InteriorRefPointCount++;				
+				}
+				else
+				{
+					VertexIsUnderwater[i] = false;
+				}
 			}
 		}
 
@@ -916,16 +954,18 @@ namespace BuoyancyAlgorithms
 		FVector IntersectionCenter;			
 
 		// we can have at most 6 intersection points with a box and a plane			
-		TArray<FVector, TInlineAllocator<FBuoyancyShapeTopologyLimits::MaxIntersectionPointsPerShape>> OrderedIntersectionPoints;		
-		TMap<int32, FVector> EdgeToIntersectionPoint;
+		TArray<FVector, TInlineAllocator<FBuoyancyShapeTopologyLimits::MaxIntersectionPointsPerShape>> OrderedIntersectionPoints;
 
 		const int32 NumEdges = ShapeQuery.NumEdges();
 		const int32 MaxBoxPlaneIntersectionsPoints = ShapeQuery.NumMaxIntersectionsPoints();
 
-		OrderedIntersectionPoints.AddUninitialized(MaxBoxPlaneIntersectionsPoints);		
+		TArray<FVector, TInlineAllocator<FBuoyancyShapeTopologyLimits::MaxEdgesPerShape>> EdgeIntersectionPoints;
+		EdgeIntersectionPoints.AddUninitialized(NumEdges);
+
+		OrderedIntersectionPoints.AddUninitialized(MaxBoxPlaneIntersectionsPoints);
 
 		int32 NumIntersectionPoints;
-		BuoyancyAlgorithms::FindAllIntersectionPoints<ShapeType>(ShapeWaterP, ShapeWaterN, ShapeQuery, WorldVertexPosition, EdgeToIntersectionPoint, NumIntersectionPoints,
+		BuoyancyAlgorithms::FindAllIntersectionPoints<ShapeType>(ShapeWaterP, ShapeWaterN, ShapeQuery, WorldVertexPosition, VertexIsUnderwater, EdgeIntersectionPoints, NumIntersectionPoints,
 			OrderedIntersectionPoints, IntersectionCenter);		
 
 		// Add intersection points to interior reference point
@@ -936,38 +976,60 @@ namespace BuoyancyAlgorithms
 		
 		const FVec3 WorldCoM = SubmergedGeneric->PCom();
 
+		// Particle velocity — constant across all triangles for this shape
+		const FVec3 ParticleV = RigidParticle->GetV();
+		const FVec3 ParticleW = RigidParticle->GetW();
+
+		// Hoist SubmergedFaceVertices allocation outside the face loop - reuse via Reset() per face
+		TArray<FVector, TInlineAllocator<FBuoyancyShapeTopologyLimits::MaxSubmergedFaceVertices>> SubmergedFaceVertices;
+
 		// Sum up submerged areas and volumes for each face
 		const int32 NumFaces = ShapeQuery.NumFaces();
 		for (int32 FaceIdx = 0; FaceIdx < NumFaces; ++FaceIdx)
 		{
-			// SCOPE_CYCLE_COUNTER(STAT_BuoyancyAlgorithms_ComputeSubmergedVolumeAndForcesForShape_FaceIteration);
-			
-			FVector SubmergedFaceCenter = FVector::ZeroVector;
-
-			// we know that a given face can have at most one more than the max number of vertices per face			
-			
-			TArray<FVector, TInlineAllocator<FBuoyancyShapeTopologyLimits::MaxSubmergedFaceVertices>> SubmergedFaceVertices;
+#if BUOYANCY_INNER_LOOP_STATS
+			SCOPE_CYCLE_COUNTER(STAT_BuoyancyAlgorithms_ComputeSubmergedVolumeAndForcesForShape_FindFaceVerticesBelowWater);
+#endif
 
-			// walk vertices for the current face in counter clockwise order and find intersections and
-			// submerged vertices		
+			// Quick pre-check: if no face vertex is underwater, the face is fully dry.
+			// With edge early-out, edges only intersect when vertices are on both sides,
+			// so no underwater vertex means no edge intersections for this face either.
 			const int32 NumVerticesForFace = ShapeQuery.NumFaceVertices(FaceIdx);
-
-			for (int32 i = 0; i < NumVerticesForFace; ++i)
 			{
-				// SCOPE_CYCLE_COUNTER(STAT_BuoyancyAlgorithms_ComputeSubmergedVolumeAndForcesForShape_FindFaceVerticesBelowWater);
+				bool bFaceHasUnderwaterVertex = false;
+				for (int32 i = 0; i < NumVerticesForFace; ++i)
+				{
+					if (VertexIsUnderwater[ShapeQuery.GetFaceVertex(FaceIdx, i)])
+					{
+						bFaceHasUnderwaterVertex = true;
+						break;
+					}
+				}
+				if (!bFaceHasUnderwaterVertex)
+				{
+					continue;
+				}
+			}
 
+			FVector SubmergedFaceCenter = FVector::ZeroVector;
+			SubmergedFaceVertices.Reset();
+
+			for (int32 i = 0; i < NumVerticesForFace; ++i)
+			{				
 				// get the current edge index belonging to the i and i+1 vertices
 				// note this is a bit weird...we could just ask the edge for the vertices but they
 				// might be in the wrong order for correctly constructing the intersected face
 				const int32 CurrEdgeIdx = ShapeQuery.GetFaceEdge(FaceIdx, i);
 
 				// if there is an intersection for this edge, add it to the list for the current face
-				const FVector *EdgeIntersection = EdgeToIntersectionPoint.Find(CurrEdgeIdx);
-				
-				if (EdgeIntersection)
+				if (CurrEdgeIdx >= 0 && CurrEdgeIdx < EdgeIntersectionPoints.Num())
 				{
-					SubmergedFaceVertices.Add(*EdgeIntersection);
-					SubmergedFaceCenter += *EdgeIntersection;
+					const FVector& EdgeIntersection = EdgeIntersectionPoints[CurrEdgeIdx];
+					if (EdgeIntersection.X != TNumericLimits<float>::Max())
+					{
+						SubmergedFaceVertices.Add(EdgeIntersection);
+						SubmergedFaceCenter += EdgeIntersection;
+					}
 				}
 
 				// Endpoint of the edge
@@ -988,13 +1050,16 @@ namespace BuoyancyAlgorithms
 			// if the face intersects the water or is fully submerged, compute face coverage
 			if (NumSubmergedFaceVertices > 0)
 			{
+#if BUOYANCY_INNER_LOOP_STATS
+				SCOPE_CYCLE_COUNTER(STAT_BuoyancyAlgorithms_ComputeSubmergedVolumeAndForcesForShape_ComputeFaceForces);
+#endif
+
 				SubmergedFaceCenter /= NumSubmergedFaceVertices;
 
 				// we have a list of the intersection points for a face in correct counterclockwise order
 				// fan layout for points on face, sum up areas and volume for the current face
 				for (int32 i = 0; i < NumSubmergedFaceVertices; i++)
-				{
-					// SCOPE_CYCLE_COUNTER(STAT_BuoyancyAlgorithms_ComputeSubmergedVolumeAndForcesForShape_ComputeFaceForces);
+				{					
 					const FVector &V0 = SubmergedFaceVertices[i];
 					const FVector &V1 = SubmergedFaceVertices[(i + 1) % NumSubmergedFaceVertices];
 					const FVector &V2 = SubmergedFaceCenter;
@@ -1003,7 +1068,16 @@ namespace BuoyancyAlgorithms
 					float TriArea;
 					float TetVolume;
 					FVector TriNormal;
-					BuoyancyAlgorithms::ComputeTriangleAreaAndVolume(V0, V1, V2, InteriorRefPoint, TriBaryCenter, TriNormal, TriArea, TetVolume, bBuoyancyDebugDraw);
+					BuoyancyAlgorithms::ComputeTriangleAreaAndVolume(V0, V1, V2, InteriorRefPoint, TriBaryCenter, TriNormal, TriArea, TetVolume);
+
+#if CHAOS_DEBUG_DRAW
+					if (bBuoyancyDebugDraw)
+					{
+						Chaos::FDebugDrawQueue::GetInstance().DrawDebugLine(V0, V1, FColor::Orange, false, -1.f, 0, 1.f);
+						Chaos::FDebugDrawQueue::GetInstance().DrawDebugLine(V1, V2, FColor::Orange, false, -1.f, 0, 1.f);
+						Chaos::FDebugDrawQueue::GetInstance().DrawDebugLine(V2, V0, FColor::Orange, false, -1.f, 0, 1.f);
+					}
+#endif
 
 					// skip degenerate triangles with 0 area
 					if (TriArea < SMALL_NUMBER)
@@ -1021,11 +1095,12 @@ namespace BuoyancyAlgorithms
 
 					// sample velocity at center of triangle
 					// #todo(dmp): this should just be interpolated from vertex velocities to avoid so many samples	
-					// #todo(dmp): split arrays out to store velocity separate from other values
-					FVector CurrWaterVelocity;					
+					// note this is tricky since there is a decent bit of bookkeeping needed with querying intersection vertices
+					FVector CurrWaterVelocity;
 					{
-						// SCOPE_CYCLE_COUNTER(STAT_BuoyancyAlgorithms_ComputeSubmergedVolumeAndForcesForShape_SampleVelocity)
-					
+#if BUOYANCY_INNER_LOOP_STATS
+						SCOPE_CYCLE_COUNTER(STAT_BuoyancyAlgorithms_ComputeSubmergedVolumeAndForcesForShape_SampleVelocity)
+#endif
 						CurrWaterVelocity = WaterSampler->SampleVelocityAtPosition(TriBaryCenter);
 					}
 
@@ -1033,7 +1108,7 @@ namespace BuoyancyAlgorithms
 					FVector TotalWorldTorque;
 
 					BuoyancyAlgorithms::ComputeFluidForceForTriangle(
-						WaterDrag, WaterDensity, DeltaSeconds, RigidParticle, WorldCoM, TriBaryCenter, TriNormal, TriArea, TetVolume,
+						WaterDrag, WaterLift, WaterDensity, DeltaSeconds, ParticleV, ParticleW, WorldCoM, TriBaryCenter, TriNormal, TriArea, TetVolume,
 						CurrWaterVelocity, ShapeWaterP, ShapeWaterN, TotalWorldForce, TotalWorldTorque);
 
 					OutForce += TotalWorldForce;
@@ -1066,7 +1141,16 @@ namespace BuoyancyAlgorithms
 			float TriArea;
 			float TetVolume;
 			FVector TriNormal;
-			BuoyancyAlgorithms::ComputeTriangleAreaAndVolume(V0, V1, V2, InteriorRefPoint, TriBaryCenter, TriNormal, TriArea, TetVolume, bBuoyancyDebugDraw);
+			BuoyancyAlgorithms::ComputeTriangleAreaAndVolume(V0, V1, V2, InteriorRefPoint, TriBaryCenter, TriNormal, TriArea, TetVolume);
+
+#if CHAOS_DEBUG_DRAW
+			if (bBuoyancyDebugDraw)
+			{
+				Chaos::FDebugDrawQueue::GetInstance().DrawDebugLine(V0, V1, FColor::Orange, false, -1.f, 0, 1.f);
+				Chaos::FDebugDrawQueue::GetInstance().DrawDebugLine(V1, V2, FColor::Orange, false, -1.f, 0, 1.f);
+				Chaos::FDebugDrawQueue::GetInstance().DrawDebugLine(V2, V0, FColor::Orange, false, -1.f, 0, 1.f);
+			}
+#endif
 
 			OutSubmergedCoM += TriBaryCenter * TriArea;
 			SubmergedCOMTotalWeight += TriArea;
@@ -1085,7 +1169,8 @@ namespace BuoyancyAlgorithms
 		// compute buoyancy force and delta velocity for solver
 		Chaos::FVec3 WorldBuoyantForce, WorldBuoyantTorque;
 		BuoyancyAlgorithms::ComputeBuoyantForceForShape(Evolution, RigidParticle, DeltaSeconds,
-			WaterDensity, OutSubmergedCoM, OutSubmergedVol, ShapeWaterN, WorldBuoyantForce, WorldBuoyantTorque);
+			WaterDensity, ParticleVol, ShapeVolGeom,
+			OutSubmergedCoM, OutSubmergedVol, ShapeWaterN, WorldBuoyantForce, WorldBuoyantTorque);
 		
 		// track the buoyancy force separately
 		OutBuoyancyForce = WorldBuoyantForce;
@@ -1109,38 +1194,55 @@ namespace BuoyancyAlgorithms
 	}
 
 	template <typename ShapeType>
-	void FindAllIntersectionPoints(const Chaos::FVec3& WaterP, const Chaos::FVec3& WaterN, const ShapeType&ShapeQuery, 
+	void FindAllIntersectionPoints(const Chaos::FVec3& WaterP, const Chaos::FVec3& WaterN, const ShapeType&ShapeQuery,
 		const TArray<FVector, TInlineAllocator<FBuoyancyShapeTopologyLimits::MaxVerticesPerShape>> &WorldVertexPosition,
-		TMap<int32, FVector> &EdgeToIntersectionPoint, int32 &NumIntersections,
+		const TArray<bool, TInlineAllocator<FBuoyancyShapeTopologyLimits::MaxVerticesPerShape>> &VertexIsUnderwater,
+		TArray<FVector, TInlineAllocator<FBuoyancyShapeTopologyLimits::MaxEdgesPerShape>>& EdgeIntersectionPoints, int32 &NumIntersections,
 		TArray<FVector, TInlineAllocator<FBuoyancyShapeTopologyLimits::MaxIntersectionPointsPerShape>>& OutOrderedIntersectionPoints, FVector& OutIntersectionCenter)
-	{		
-		// SCOPE_CYCLE_COUNTER(STAT_BuoyancyAlgorithms_FindAllIntersectionPoints)
+	{
+		SCOPE_CYCLE_COUNTER(STAT_BuoyancyAlgorithms_FindAllIntersectionPoints)
 
 		OutIntersectionCenter = FVector::ZeroVector;
+		NumIntersections = 0;
+
+		// Sentinel value for non-intersecting edges
+		const FVector SentinelValue(TNumericLimits<float>::Max());
 
 		// loop over all edges and compute intersection points
 		const int32 NumEdges = ShapeQuery.NumEdges();
 		for (int32 CurrEdgeIdx = 0; CurrEdgeIdx < NumEdges; ++CurrEdgeIdx)
-		{		
+		{
 			int32 Index0, Index1;
 			ShapeQuery.GetEdgeVertices(CurrEdgeIdx, Index0, Index1);
 
+			// An edge can only intersect the water plane if one vertex is above and one below.
+			// Skip edges where both vertices are on the same side - avoids EdgePlaneIntersection math.
+			if (VertexIsUnderwater[Index0] == VertexIsUnderwater[Index1])
+			{
+				EdgeIntersectionPoints[CurrEdgeIdx] = SentinelValue;
+				continue;
+			}
+
+			// Use a local Chaos::FVec3 (float) since EdgePlaneIntersection takes float precision
 			Chaos::FVec3 IntersectionPoint;
 			if (EdgePlaneIntersection(WaterP, WaterN, WorldVertexPosition[Index0], WorldVertexPosition[Index1], IntersectionPoint))
 			{
-				EdgeToIntersectionPoint.Add(CurrEdgeIdx, IntersectionPoint);
-				OutIntersectionCenter += IntersectionPoint;				
+				EdgeIntersectionPoints[CurrEdgeIdx] = FVector(IntersectionPoint);
+				OutIntersectionCenter += EdgeIntersectionPoints[CurrEdgeIdx];
+				NumIntersections++;
+			}
+			else
+			{
+				EdgeIntersectionPoints[CurrEdgeIdx] = SentinelValue;
 			}
 		}
 
-		NumIntersections = EdgeToIntersectionPoint.Num();
-
 		if (NumIntersections > 0)
 		{
 			OutIntersectionCenter /= NumIntersections;
 
 			// Sort intersection points by angle around reference point
-			BuoyancyAlgorithms::SortIntersectionPointsByAngle<ShapeType>(WaterP, WaterN, OutIntersectionCenter, ShapeQuery, EdgeToIntersectionPoint, OutOrderedIntersectionPoints);			
+			BuoyancyAlgorithms::SortIntersectionPointsByAngle(WaterN, OutIntersectionCenter, EdgeIntersectionPoints, NumIntersections, OutOrderedIntersectionPoints);
 		}
 		else
 		{
@@ -1148,32 +1250,44 @@ namespace BuoyancyAlgorithms
 		}
 	}
 
-	template <typename ShapeType>
-	void SortIntersectionPointsByAngle(const Chaos::FVec3& WaterP, const Chaos::FVec3& WaterN, const Chaos::FVec3& IntersectionCenter, const ShapeType& ShapeQuery,
-		const TMap<int, FVector>& EdgeToIntersectionPoint, TArray<FVector, TInlineAllocator<FBuoyancyShapeTopologyLimits::MaxIntersectionPointsPerShape>>& OutOrderedIntersectionPoints)
+	void SortIntersectionPointsByAngle(const Chaos::FVec3& WaterN, const Chaos::FVec3& IntersectionCenter,
+		const TArray<FVector, TInlineAllocator<FBuoyancyShapeTopologyLimits::MaxEdgesPerShape>>& EdgeIntersectionPoints,
+		int32 NumIntersections,
+		TArray<FVector, TInlineAllocator<FBuoyancyShapeTopologyLimits::MaxIntersectionPointsPerShape>>& OutOrderedIntersectionPoints)
 	{
 		SCOPE_CYCLE_COUNTER(STAT_BuoyancyAlgorithms_SortIntersectionPointsByAngle)
 
-		// Order intersection points around intersection center
-		TMap<double, FVector> IntersectionPointAngleToCenter;
-				
-		bool FoundFirst = false;
+		// Sort directly into the output array with a parallel angle array for ordering.
+		// Insertion sort is optimal for the typical 3-6 intersection points.
+		TArray<double, TInlineAllocator<FBuoyancyShapeTopologyLimits::MaxIntersectionPointsPerShape>> Angles;
+		Angles.Reserve(NumIntersections);
+		OutOrderedIntersectionPoints.Reset();
+		OutOrderedIntersectionPoints.Reserve(NumIntersections);
 
 		// first element is assigned to be angle 0
-		FVector AngleStart;		
-		
-		// loop over all edges, find intersection points then determine angle to reference point		
-		for (const TPair<int32, FVector>& IntersectedEdge : EdgeToIntersectionPoint)
+		FVector AngleStart;
+		bool FoundFirst = false;
+
+		// loop over all edges, find valid intersection points then determine angle to reference point
+		const int32 NumEdges = EdgeIntersectionPoints.Num();
+		for (int32 EdgeIdx = 0; EdgeIdx < NumEdges; ++EdgeIdx)
 		{
+			const FVector& CurrPoint = EdgeIntersectionPoints[EdgeIdx];
+
+			// Skip edges with no intersection (sentinel value)
+			if (CurrPoint.X == TNumericLimits<float>::Max())
+			{
+				continue;
+			}
+
 			double CurrAngle;
-			const FVector& CurrPoint = IntersectedEdge.Value;
-			
+
 			if (!FoundFirst) // first point found is the 0 angle point
 			{
 				AngleStart = CurrPoint - IntersectionCenter;
 				AngleStart.Normalize();
 
-				CurrAngle = 0.f;
+				CurrAngle = 0.0;
 				FoundFirst = true;
 			}
 			else // compute angle from the first point to this one
@@ -1181,25 +1295,27 @@ namespace BuoyancyAlgorithms
 				FVector CurrVector = CurrPoint - IntersectionCenter;
 				CurrVector.Normalize();
 
-				CurrAngle = FMath::Acos(AngleStart.Dot(CurrVector));
+				CurrAngle = FMath::Acos(FMath::Clamp(AngleStart.Dot(CurrVector), -1.0, 1.0));
 
 				if (AngleStart.Cross(CurrVector).Dot(WaterN) < SMALL_NUMBER)
 				{
-					CurrAngle = 2.f * PI - CurrAngle;
+					CurrAngle = 2.0 * PI - CurrAngle;
 				}
 			}
 
-			IntersectionPointAngleToCenter.Add(CurrAngle, CurrPoint);
+			// Insertion sort into both parallel arrays
+			int32 InsertIdx = Angles.Num();
+			for (int32 j = 0; j < Angles.Num(); ++j)
+			{
+				if (CurrAngle < Angles[j])
+				{
+					InsertIdx = j;
+					break;
+				}
+			}
+			Angles.Insert(CurrAngle, InsertIdx);
+			OutOrderedIntersectionPoints.Insert(CurrPoint, InsertIdx);
 		}
-
-		// #todo(dmp): replace this sort with something faster since we only have between 3-6 points to sort (and index 0 is already sorted) for boxes
-		// order intersection points by the angle it makes so we have a correct winding order for tesselation
-		IntersectionPointAngleToCenter.KeySort([](const double& A, const double& B) {
-			return A < B;
-			});
-		
-		// Generate an array of all intersection points		
-		IntersectionPointAngleToCenter.GenerateValueArray(OutOrderedIntersectionPoints);
 	}
 
 	bool EdgePlaneIntersection(const Chaos::FVec3& WaterP, const Chaos::FVec3& WaterN, const Chaos::FVec3& V0, const Chaos::FVec3& V1, Chaos::FVec3& IntersectionPoint)
@@ -1227,97 +1343,10 @@ namespace BuoyancyAlgorithms
 		}
 
 		return false;
-	}
-
-	void ComputeTriangleAreaAndVolume(const FVector &V0, const FVector &V1, const FVector &V2, const FVector &MeshCenterPoint, FVector &OutTriangleBaryCenter, FVector &OutNormal, float& OutArea, float& OutVolume, bool DebugDraw /*= false*/)
-	{
-		// SCOPE_CYCLE_COUNTER(STAT_BuoyancyAlgorithms_ComputeTriangleAreaAndVolume)
-
-		// compute center of triangle
-		OutTriangleBaryCenter = (V0 + V1 + V2) / 3.;
-		
-		// add up the volume of the tet created by this triangle and the center of the box
-		const FVector A = V0 - MeshCenterPoint;
-		const FVector B = V1 - MeshCenterPoint;
-		const FVector C = V2 - MeshCenterPoint;
-		const FVector N = A.Cross(B);
-
-		OutVolume = FMath::Abs((N.Dot(C)) / 6.f);
-
-		OutNormal = (V1 - V0).Cross(V2 - V0);
-		float NormalLength = OutNormal.Length();
-		
-		// #todo(dmp): careful w/ divide by zero here?
-		OutNormal /= NormalLength;
-
-		// area of the triangle
-		OutArea = .5 * NormalLength;
-		
-#if CHAOS_DEBUG_DRAW
-		if (DebugDraw)
-		{
-			Chaos::FDebugDrawQueue::GetInstance().DrawDebugLine(V0, V1, FColor::Orange, false, -1.f, 0, 1.f);
-			Chaos::FDebugDrawQueue::GetInstance().DrawDebugLine(V1, V2, FColor::Orange, false, -1.f, 0, 1.f);
-			Chaos::FDebugDrawQueue::GetInstance().DrawDebugLine(V2, V0, FColor::Orange, false, -1.f, 0, 1.f);
-
-			//Chaos::FDebugDrawQueue::GetInstance().DrawDebugSphere(OutTriangleBaryCenter, 5, 10, FColor::Cyan, false, -1.f, -1, 1.f);
-			//Chaos::FDebugDrawQueue::GetInstance().DrawDebugDirectionalArrow(OutTriangleBaryCenter, OutTriangleBaryCenter + OutNormal * 20, 20, FColor::Cyan, false, -1.f, -1, 2.f);
-		}
-#endif
-	}
-
-	// #todo(dmp): optimization- precalculate particle v, w, m given they never change between calls to compute forces
-	void ComputeFluidForceForTriangle(const float WaterDrag,
-		const float DeltaSeconds, const float WaterDensity,
-		const Chaos::FPBDRigidParticleHandle* RigidParticle, const FVector WorldCoM,
-		const FVector &TriBaryCenter, const FVector &TriNormal, const float TriArea, const float TetVolume,
-		const FVector &WaterVelocity, const FVector &WaterP, const FVector &WaterN,
-		FVector &OutTotalWorldForce, FVector &OutTotalWorldTorque)
-	{
-		// SCOPE_CYCLE_COUNTER(STAT_BuoyancyAlgorithms_ComputeFluidForceForTriangle)
-
-		OutTotalWorldForce = FVector(0, 0, 0);
-		OutTotalWorldTorque = FVector(0, 0, 0);
-
-		const FVec3 WorldForcePosition = TriBaryCenter;
-		const FVec3 WorldCOMToForcePos = WorldForcePosition - WorldCoM;
-
-		// project velocity sample onto water plane to support waterfalls and flowing rivers more accurately
-		const FVec3 WaterVelocityOnPlane = WaterVelocity - WaterVelocity.Dot(WaterN) * WaterN;
-
-		// Get world space particle linear velocity at current point
-		// note we are including the linear velocity from torque so objects spin properly in flow
-		const FVec3 SubmergedParticleVelocity = RigidParticle->GetV() + FVec3::CrossProduct(RigidParticle->GetW(), WorldCOMToForcePos);
-
-		// compute force and torque to set linear velocity to fluid velocity				
-		const FVec3 RelativeVelocity = WaterVelocityOnPlane - SubmergedParticleVelocity;
-		const float RelativeVelocityMag = RelativeVelocity.Length();
-
-		// test if this triangle is influenced by the water based on the normal since we have closed shapes only.  This drag algorithm
-		// is derived for two sided planes
-		const float FacingTest = RelativeVelocity.Dot(TriNormal);
-
-		if (RelativeVelocityMag < SMALL_NUMBER || FacingTest < 0.f)
-		{
-			return;
-		}
-
-		// go with the flow combined drag and lift
-		// https://www.yousufsoliman.com/projects/download/going-with-the-flow.pdf
-		//const FVec3 ForceFromWaterVelocity = .5 * WaterDensity * RelativeVelocityMag * RelativeVelocity.Dot(TriNormal) * TriNormal * WaterDrag * TriArea;	
-
-		// https://www.cemyuksel.com/research/waveparticles/cem_yuksel_dissertation.pdf
-		const float A = TriArea * RelativeVelocity.Dot(TriNormal) / RelativeVelocityMag;
-		const FVec3 ForceFromWaterVelocity = .5 * WaterDensity * RelativeVelocityMag * RelativeVelocity * A * WaterDrag;	
-
-		// compute torque based on the linear force we apply
-		const FVec3 TorqueFromWaterVelocity = Chaos::FVec3::CrossProduct(WorldCOMToForcePos, ForceFromWaterVelocity);
-
-		OutTotalWorldForce += ForceFromWaterVelocity;
-		OutTotalWorldTorque += TorqueFromWaterVelocity;
-	}
+	}	
 
 	void ComputeBuoyantForceForShape(const Chaos::FPBDRigidsEvolution& Evolution, const Chaos::FPBDRigidParticleHandle* RigidParticle, const float DeltaSeconds, const float WaterDensity,
+		const float ParticleVol, const float ShapeVolGeom,
 		const Chaos::FVec3& SubmergedCoM, const float SubmergedVol, const Chaos::FVec3& WaterN, Chaos::FVec3& OutWorldBuoyantForce, Chaos::FVec3& OutWorldBuoyantTorque)
 	{
 		SCOPE_CYCLE_COUNTER(STAT_BuoyancyAlgorithms_ComputeBuoyantForceForShape)
@@ -1337,10 +1366,10 @@ namespace BuoyancyAlgorithms
 		const FVec3 GravityDir = FVec3::DownVector;
 		const float GravityAccel = FVec3::DotProduct(GravityDir, GravityAccelVec);
 
-		// Scale submerged volume
+		// Scale submerged volume using precomputed particle/shape volumes
 		float ScaledSubmergedVol = SubmergedVol;
 		float TotalParticleVol;
-		ScaleSubmergedVolume(Evolution, RigidParticle, false, ScaledSubmergedVol, TotalParticleVol);
+		ScaleSubmergedVolume(ParticleVol, ShapeVolGeom, ScaledSubmergedVol, TotalParticleVol);
 
 		// Compute buoyant force
 		//
@@ -1379,7 +1408,9 @@ namespace BuoyancyAlgorithms
 	}
 
 	void FBuoyancyConvexShape::Initialize()
-	{		
+	{
+		SCOPE_CYCLE_COUNTER(STAT_BuoyancyAlgorithms_FBuoyancyConvexShape_Initialize)
+
 		// #todo(dmp): ideally the convex would have a mapping between half edges and edges that is serialized, but instead for now we'll build it here
 		const int32 NumEdges = Convex->NumEdges();
 		HalfEdgeToEdge.AddUninitialized(Convex->GetStructureData().NumHalfEdges());
diff --git a/Engine/Plugins/Experimental/Buoyancy/Source/Runtime/Private/BuoyancySubsystem.cpp b/Engine/Plugins/Experimental/Buoyancy/Source/Runtime/Private/BuoyancySubsystem.cpp
--- a/Engine/Plugins/Experimental/Buoyancy/Source/Runtime/Private/BuoyancySubsystem.cpp
+++ b/Engine/Plugins/Experimental/Buoyancy/Source/Runtime/Private/BuoyancySubsystem.cpp
@@ -250,6 +250,7 @@ void UBuoyancySubsystem::ApplyRuntimeSettings(const UBuoyancyRuntimeSettings* In
 	// so introduce a factor of 10^-3 here.
 	BuoyancySettings.WaterDensity = Chaos::GCm3ToKgCm3(InSettings->WaterDensity);
 	BuoyancySettings.WaterDrag = InSettings->WaterDrag;
+	BuoyancySettings.WaterLift = InSettings->WaterLift;
 	BuoyancySettings.WaterCollisionChannel = InSettings->CollisionChannelForWaterObjects;
 	BuoyancySettings.bKeepAwake = InSettings->bKeepFloatingObjectsAwake;
 	BuoyancySettings.MaxNumBoundsSubdivisions = InSettings->MaxNumBoundsSubdivisions;
@@ -1205,6 +1206,7 @@ void FBuoyancySubsystemSimCallback::ProcessAccurateInteraction(Chaos::FPBDRigids
 	float TotalParticleVol = 0;
 	
 	const float AccurateIntegrationDrag = AccurateIntegrationDragMultiplier * BuoyancySettings->WaterDrag;
+	const float AccurateIntegrationLift = AccurateIntegrationDragMultiplier * BuoyancySettings->WaterLift;
 
 	if (!Interaction.WaterSampler.IsValid())
 	{
@@ -1216,7 +1218,7 @@ void FBuoyancySubsystemSimCallback::ProcessAccurateInteraction(Chaos::FPBDRigids
 		BuoyancyAlgorithms::ComputeSubmergedVolumeAndForcesForParticle(BuoyancyParticleData,
 			Interaction.RigidParticle, Interaction.WaterParticle,
 			StaticCastSharedPtr<FBuoyancyConstantSplineSampler>(Interaction.WaterSampler),
-			Evolution, DeltaSeconds, BuoyancySettings->WaterDensity, AccurateIntegrationDrag,
+			Evolution, DeltaSeconds, BuoyancySettings->WaterDensity, AccurateIntegrationDrag, AccurateIntegrationLift,
 			TotalParticleVol, SubmergedVol, SubmergedCoM, WorldForce, WorldTorque);		
 	}
 	else if (Interaction.WaterSampler->GetSamplerType() == EWaterSamplerType::ShallowWater)
@@ -1224,7 +1226,7 @@ void FBuoyancySubsystemSimCallback::ProcessAccurateInteraction(Chaos::FPBDRigids
 		BuoyancyAlgorithms::ComputeSubmergedVolumeAndForcesForParticle(BuoyancyParticleData,
 			Interaction.RigidParticle, Interaction.WaterParticle,
 			StaticCastSharedPtr<FBuoyancyShallowWaterSampler>(Interaction.WaterSampler),
-			Evolution, DeltaSeconds, BuoyancySettings->WaterDensity, AccurateIntegrationDrag,
+			Evolution, DeltaSeconds, BuoyancySettings->WaterDensity, AccurateIntegrationDrag, AccurateIntegrationLift,
 			TotalParticleVol, SubmergedVol, SubmergedCoM, WorldForce, WorldTorque);
 	}
 
diff --git a/Engine/Plugins/Experimental/Buoyancy/Source/Runtime/Public/BuoyancyAlgorithms.h b/Engine/Plugins/Experimental/Buoyancy/Source/Runtime/Public/BuoyancyAlgorithms.h
--- a/Engine/Plugins/Experimental/Buoyancy/Source/Runtime/Public/BuoyancyAlgorithms.h
+++ b/Engine/Plugins/Experimental/Buoyancy/Source/Runtime/Public/BuoyancyAlgorithms.h
@@ -40,10 +40,10 @@ namespace BuoyancyAlgorithms
 	struct FBuoyancyShapeTopologyLimits
 	{
 	public:
-		static const int32 MaxVerticesPerShape = 100;		
-		static const int32 MaxEdgesPerShape = 100;
+		static const int32 MaxVerticesPerShape = 50;		
+		static const int32 MaxEdgesPerShape = 144;
 		static const int32 MaxIntersectionPointsPerShape = MaxEdgesPerShape;
-		static const int32 MaxVerticesPerFace = 20;
+		static const int32 MaxVerticesPerFace = 10;
 		static const int32 MaxSubmergedFaceVertices = MaxVerticesPerFace + 1;
 	};
 	
@@ -181,7 +181,7 @@ namespace BuoyancyAlgorithms
 	void ComputeSubmergedVolumeAndForcesForParticle(FBuoyancyParticleData& ParticleData,
 		const Chaos::FGeometryParticleHandle* SubmergedParticle, const Chaos::FGeometryParticleHandle* WaterParticle,
 		TSharedPtr<SamplerType> WaterSampler,
-		const Chaos::FPBDRigidsEvolution& Evolution, const float DeltaSeconds, const float WaterDensity, const float WaterDrag,
+		const Chaos::FPBDRigidsEvolution& Evolution, const float DeltaSeconds, const float WaterDensity, const float WaterDrag, const float WaterLift,
 		float& OutTotalParticleVol, float& OutTotalSubmergedVol, Chaos::FVec3& OutTotalSubmergedCoM, Chaos::FVec3& OutTotalForce, Chaos::FVec3& OutTotalTorque);
 
 	// given a shape, compute the submerged volume and accumulate forces
@@ -189,7 +189,8 @@ namespace BuoyancyAlgorithms
 	template <typename ShapeType, typename SamplerType>
 	void ComputeSubmergedVolumeAndForcesForShape(
 		const Chaos::FGeometryParticleHandle* SubmergedParticle, const ShapeType& BoxShape,
-		const Chaos::FPBDRigidsEvolution& Evolution, float DeltaSeconds, const float WaterDensity, const float WaterDrag,
+		const Chaos::FPBDRigidsEvolution& Evolution, float DeltaSeconds, const float WaterDensity, const float WaterDrag, const float WaterLift,
+		const float ParticleVol, const float ShapeVolGeom,
 		const Chaos::FRigidTransform3 &ShapeWorldTransform, TSharedPtr<SamplerType> WaterSampler,
 		float& OutSubmergedVol, Chaos::FVec3& OutSubmergedCoM,
 		Chaos::FVec3& OutForce, Chaos::FVec3& OutTorque,
@@ -199,30 +200,106 @@ namespace BuoyancyAlgorithms
 	template <typename ShapeType>
 	void FindAllIntersectionPoints(const Chaos::FVec3& WaterP, const Chaos::FVec3& WaterN, const ShapeType& BoxShape,
 		const TArray<FVector, TInlineAllocator<FBuoyancyShapeTopologyLimits::MaxVerticesPerShape>> &WorldVertexPosition,
-		TMap<int32, FVector>& EdgeToIntersectionPoint, int32& NumIntersections,
+		const TArray<bool, TInlineAllocator<FBuoyancyShapeTopologyLimits::MaxVerticesPerShape>> &VertexIsUnderwater,
+		TArray<FVector, TInlineAllocator<FBuoyancyShapeTopologyLimits::MaxEdgesPerShape>>& EdgeIntersectionPoints,
+		int32& NumIntersections,
 		TArray<FVector, TInlineAllocator<FBuoyancyShapeTopologyLimits::MaxIntersectionPointsPerShape>>& OutOrderedIntersectionPoints, FVector& OutIntersectionCenter);
 
 	// sort intersection points by angle
-	template <typename ShapeType>
-	void SortIntersectionPointsByAngle(const Chaos::FVec3& WaterP, const Chaos::FVec3& WaterN, const Chaos::FVec3& IntersectionCenter, const ShapeType& BoxShape,
-		const TMap<int, FVector>& EdgeToIntersectionPoint,
+	void SortIntersectionPointsByAngle(const Chaos::FVec3& WaterN, const Chaos::FVec3& IntersectionCenter,
+		const TArray<FVector, TInlineAllocator<FBuoyancyShapeTopologyLimits::MaxEdgesPerShape>>& EdgeIntersectionPoints,
+		int32 NumIntersections,
 		TArray<FVector, TInlineAllocator<FBuoyancyShapeTopologyLimits::MaxIntersectionPointsPerShape>>& OutOrderedIntersectionPoints);
 
 	bool EdgePlaneIntersection(const Chaos::FVec3& WaterP, const Chaos::FVec3& WaterN, const Chaos::FVec3& V0, const Chaos::FVec3& V1, Chaos::FVec3& IntersectionPoint);
 
 	// compute area and volume of a tet from a triangle and center point on mesh
-	void ComputeTriangleAreaAndVolume(const FVector &V0, const FVector &V1, const FVector &V2,
-		const FVector &MeshCenterPoint, FVector& OutTriangleBaryCenter, FVector& OutNormal, float& OutArea, float& OutVolume, bool DebugDraw = false);
+	FORCEINLINE void ComputeTriangleAreaAndVolume(const FVector &V0, const FVector &V1, const FVector &V2,
+		const FVector &MeshCenterPoint, FVector& OutTriangleBaryCenter, FVector& OutNormal, float& OutArea, float& OutVolume)
+	{
+		// compute center of triangle
+		OutTriangleBaryCenter = (V0 + V1 + V2) / 3.;
+
+		// add up the volume of the tet created by this triangle and the center of the box
+		const FVector A = V0 - MeshCenterPoint;
+		const FVector B = V1 - MeshCenterPoint;
+		const FVector C = V2 - MeshCenterPoint;
+		const FVector N = A.Cross(B);
+
+		OutVolume = FMath::Abs((N.Dot(C)) / 6.f);
+
+		OutNormal = (V1 - V0).Cross(V2 - V0);
+		float NormalLength = OutNormal.Length();
+
+		// #todo(dmp): careful w/ divide by zero here?
+		OutNormal /= NormalLength;
+
+		// area of the triangle
+		OutArea = .5 * NormalLength;
+	}
 
 	// compute the force the fluid exerts on a triangle
-	void ComputeFluidForceForTriangle(const float WaterDrag,
-		const float DeltaSeconds, const float WaterDensity,
-		const Chaos::FPBDRigidParticleHandle* RigidParticle, const FVector WorldCoM,
+	FORCEINLINE void ComputeFluidForceForTriangle(const float WaterDrag, const float WaterLift,
+		const float WaterDensity, const float DeltaSeconds,
+		const Chaos::FVec3& ParticleV, const Chaos::FVec3& ParticleW, const FVector WorldCoM,
 		const FVector &TriBaryCenter, const FVector &TriNormal, const float TriArea, const float TetVolume,
 		const FVector &WaterVelocity, const FVector &WaterP, const FVector &WaterN,
-		FVector& OutTotalWorldForce, FVector& OutTotalWorldTorque);
+		FVector& OutTotalWorldForce, FVector& OutTotalWorldTorque)
+	{
+		OutTotalWorldForce = FVector(0, 0, 0);
+		OutTotalWorldTorque = FVector(0, 0, 0);
+
+		const Chaos::FVec3 WorldForcePosition = TriBaryCenter;
+		const Chaos::FVec3 WorldCOMToForcePos = WorldForcePosition - WorldCoM;
+
+		// project velocity sample onto water plane to support waterfalls and flowing rivers more accurately
+		const Chaos::FVec3 WaterVelocityOnPlane = WaterVelocity - WaterVelocity.Dot(WaterN) * WaterN;
+
+		// Get world space particle linear velocity at current point
+		// note we are including the linear velocity from torque so objects spin properly in flow
+		const Chaos::FVec3 SubmergedParticleVelocity = ParticleV + Chaos::FVec3::CrossProduct(ParticleW, WorldCOMToForcePos);
+
+		// compute force and torque to set linear velocity to fluid velocity
+		const Chaos::FVec3 RelativeVelocity = WaterVelocityOnPlane - SubmergedParticleVelocity;
+		const float RelativeVelocityMag = RelativeVelocity.Length();
+
+		// test if this triangle is influenced by the water based on the normal since we have closed shapes only.  This drag algorithm
+		// is derived for two sided planes
+		const float FacingTest = RelativeVelocity.Dot(TriNormal);
+
+		if (RelativeVelocityMag < SMALL_NUMBER || FacingTest < 0.f)
+		{
+			return;
+		}
+
+		// go with the flow combined drag and lift
+		// https://www.yousufsoliman.com/projects/download/going-with-the-flow.pdf
+		// const FVec3 ForceFromWaterVelocity = .5 * WaterDensity * RelativeVelocityMag * RelativeVelocity.Dot(TriNormal) * TriNormal * WaterDrag * TriArea;	
+
+		// https://www.cemyuksel.com/research/waveparticles/cem_yuksel_dissertation.pdf
+		const float FluidA = TriArea * RelativeVelocity.Dot(TriNormal) / RelativeVelocityMag;
+		const Chaos::FVec3 DragForce = .5 * WaterDensity * RelativeVelocityMag * RelativeVelocity * FluidA * WaterDrag;
+
+		const Chaos::FVec3 LiftCrossProd = TriNormal.Cross(RelativeVelocity);
+		const float LiftCrossProdLength = LiftCrossProd.Length();
+		
+		Chaos::FVec3 LiftForce = {0.0, 0.0, 0.0};		
+		if (LiftCrossProdLength > SMALL_NUMBER)
+		{
+			LiftForce = .5 * WaterDensity * RelativeVelocityMag * RelativeVelocity.Cross(LiftCrossProd / LiftCrossProdLength) * FluidA * WaterLift;
+		}
+
+		const Chaos::FVec3 ForceFromWaterVelocity = DragForce + LiftForce;
+
+		// compute torque based on the linear force we apply
+		const Chaos::FVec3 TorqueFromWaterVelocity = Chaos::FVec3::CrossProduct(WorldCOMToForcePos, ForceFromWaterVelocity);
+
+		OutTotalWorldForce += ForceFromWaterVelocity;
+		OutTotalWorldTorque += TorqueFromWaterVelocity;
+	}
 
 	// compute the force the buoyancy exerts on a shape
 	void ComputeBuoyantForceForShape(const Chaos::FPBDRigidsEvolution& Evolution, const Chaos::FPBDRigidParticleHandle* RigidParticle, const float DeltaSeconds, const float WaterDensity,
+		const float ParticleVol, const float ShapeVolGeom,
 		const Chaos::FVec3& SubmergedCoM, const float SubmergedVol, const Chaos::FVec3& WaterN, Chaos::FVec3& OutWorldBuoyantForce, Chaos::FVec3& OutWorldBuoyantTorque);
 }
diff --git a/Engine/Plugins/Experimental/Buoyancy/Source/Runtime/Public/BuoyancyRuntimeSettings.h b/Engine/Plugins/Experimental/Buoyancy/Source/Runtime/Public/BuoyancyRuntimeSettings.h
--- a/Engine/Plugins/Experimental/Buoyancy/Source/Runtime/Public/BuoyancyRuntimeSettings.h
+++ b/Engine/Plugins/Experimental/Buoyancy/Source/Runtime/Public/BuoyancyRuntimeSettings.h
@@ -40,6 +40,12 @@ public:
 	UPROPERTY(EditAnywhere, Config, Category = WaterProperties, Meta = (ClampMin = 0))
 	float WaterDrag = 1.f;
 
+	/** Lift coefficient for submerged surfaces. Controls how much angled surfaces
+	    deflect objects perpendicular to the flow (rudder steering, planing).
+	    0 = no lift (current behavior), 1 = standard lift. */
+	UPROPERTY(EditAnywhere, Config, Category = WaterProperties, Meta = (ClampMin = 0))
+	float WaterLift = 0.f;
+
 	/** Collision channel to use for water ObjectTypes */
 	UPROPERTY(EditAnywhere, Config, Category = WaterProperties)
 	TEnumAsByte<ECollisionChannel> CollisionChannelForWaterObjects;
diff --git a/Engine/Plugins/Experimental/Buoyancy/Source/Runtime/Public/BuoyancyStats.h b/Engine/Plugins/Experimental/Buoyancy/Source/Runtime/Public/BuoyancyStats.h
--- a/Engine/Plugins/Experimental/Buoyancy/Source/Runtime/Public/BuoyancyStats.h
+++ b/Engine/Plugins/Experimental/Buoyancy/Source/Runtime/Public/BuoyancyStats.h
@@ -48,6 +48,8 @@ DECLARE_CYCLE_STAT(TEXT("Buoyancy_Algorithms_SortIntersectionPointsByAngle"), ST
 DECLARE_CYCLE_STAT(TEXT("Buoyancy_Algorithms_ComputeTriangleAreaAndVolume"), STAT_BuoyancyAlgorithms_ComputeTriangleAreaAndVolume, STATGROUP_Buoyancy);
 DECLARE_CYCLE_STAT(TEXT("Buoyancy_Algorithms_ComputeFluidForceForTriangle"), STAT_BuoyancyAlgorithms_ComputeFluidForceForTriangle, STATGROUP_Buoyancy);
 DECLARE_CYCLE_STAT(TEXT("Buoyancy_Algorithms_ComputeBuoyantForceForShape"), STAT_BuoyancyAlgorithms_ComputeBuoyantForceForShape, STATGROUP_Buoyancy);
+DECLARE_CYCLE_STAT(TEXT("Buoyancy_Algorithms_FBuoyancyConvexShape_Initialize"), STAT_BuoyancyAlgorithms_FBuoyancyConvexShape_Initialize, STATGROUP_Buoyancy);
+DECLARE_CYCLE_STAT(TEXT("Buoyancy_Algorithms_ComputeSubmergedVolumeAndForcesForShape_FindFaceVerticesBelowWater"), STAT_BuoyancyAlgorithms_ComputeSubmergedVolumeAndForcesForShape_FindFaceVerticesBelowWater, STATGROUP_Buoyancy);
 
 // Particle Data
 DECLARE_CYCLE_STAT(TEXT("Buoyancy_ParticleData_GetData"), STAT_BuoyancyParticleData_GetData, STATGROUP_Buoyancy);
diff --git a/Engine/Plugins/Experimental/Buoyancy/Source/Runtime/Public/BuoyancySubsystem.h b/Engine/Plugins/Experimental/Buoyancy/Source/Runtime/Public/BuoyancySubsystem.h
--- a/Engine/Plugins/Experimental/Buoyancy/Source/Runtime/Public/BuoyancySubsystem.h
+++ b/Engine/Plugins/Experimental/Buoyancy/Source/Runtime/Public/BuoyancySubsystem.h
@@ -58,6 +58,8 @@ struct FBuoyancySettings
 
 	float WaterDrag = 1.f; // unitless
 
+	float WaterLift = 0.f; // unitless
+
 	int32 MaxNumBoundsSubdivisions = 2;
 
 	float MinBoundsSubdivisionVol = FMath::Pow(125.f, 3.f); // 1m^3
