- redo scaling on drag (1 now is correct and works well)
- several other small fixes
- Support for waves
- fixed interactions with oceans
- small cleanup and bug fixes

#Synced-CL 52589973

diff --git a/Engine/Plugins/Experimental/Buoyancy/Source/Runtime/Private/BuoyancyAlgorithms.cpp b/Engine/Plugins/Experimental/Buoyancy/Source/Runtime/Private/BuoyancyAlgorithms.cpp
--- a/Engine/Plugins/Experimental/Buoyancy/Source/Runtime/Private/BuoyancyAlgorithms.cpp
+++ b/Engine/Plugins/Experimental/Buoyancy/Source/Runtime/Private/BuoyancyAlgorithms.cpp
@@ -249,6 +249,48 @@ namespace BuoyancyAlgorithms
 		return ShapeVol;
 	}
 
+	void ComputeShapeVolumes(const FGeometryParticleHandle* Particle, FRealSingle& OutBBoxVol, FRealSingle& OutGeomVol)
+	{
+		OutBBoxVol = 0.f;
+		OutGeomVol = 0.f;
+
+		if (Particle == nullptr)
+		{
+			return;
+		}
+
+		const FImplicitObject* ImplicitObject = Particle->GetGeometry();
+		if (ImplicitObject == nullptr)
+		{
+			return;
+		}
+
+		const FShapeInstanceArray& ShapeInstances = Particle->ShapeInstances();
+		if (ShapeInstances.Num() == 0)
+		{
+			return;
+		}
+
+		ImplicitObject->VisitLeafObjects(
+			[Particle, &ShapeInstances, &OutBBoxVol, &OutGeomVol]
+			(const FImplicitObject* InnerImplicitObject, const FRigidTransform3&, const int32 RootObjectIndex, const int32, const int32)
+		{
+			const int32 ShapeIndex = ShapeInstances.IsValidIndex(RootObjectIndex) ? RootObjectIndex : 0;
+			const EImplicitObjectType ShapeType = Chaos::Private::GetImplicitCollisionType(Particle, InnerImplicitObject);
+			if (DoCollide(ShapeType, ShapeInstances[ShapeIndex].Get()))
+			{
+				Chaos::Utilities::CastHelper(*InnerImplicitObject, [&OutBBoxVol, &OutGeomVol](const auto& Geom)
+				{
+					const FRealSingle BBoxVol = Geom.BoundingBox().GetVolume();
+					const FRealSingle GeomVol = Geom.GetVolume();
+					OutBBoxVol += BBoxVol;
+					// Some implicit types return 0 from GetVolume(); fall back to BBox volume
+					OutGeomVol += (GeomVol > UE_SMALL_NUMBER) ? GeomVol : BBoxVol;
+				});
+			}
+		});
+	}
+
 	void ScaleSubmergedVolume(const FPBDRigidsEvolutionGBF& Evolution, const FGeometryParticleHandle* Particle, const bool UseBoundingBoxVolume, FRealSingle& SubmergedVol, FRealSingle& TotalVol)
 	{
 		SCOPE_CYCLE_COUNTER(STAT_BuoyancyAlgorithms_ScaleSubmergedVolume)
@@ -699,13 +741,13 @@ namespace BuoyancyAlgorithms
 		// Precompute particle and shape volumes — both are constant per particle per tick
 		// and would otherwise be recomputed via VisitLeafObjects on every shape + once per particle.
 		const FRealSingle ParticleVol = ComputeParticleVolume(Evolution, SubmergedParticle);
-		const FRealSingle ShapeVolBBox = ComputeShapeVolume(SubmergedParticle, true);
-		const FRealSingle ShapeVolGeom = ComputeShapeVolume(SubmergedParticle, false);
+		FRealSingle ShapeVolBBox, ShapeVolGeom;
+		ComputeShapeVolumes(SubmergedParticle, ShapeVolBBox, ShapeVolGeom);
 
 		// Traverse the submerged particle's leaves
 		RootImplicit->VisitLeafObjects(
 			[&Evolution, DeltaSeconds, WaterDensity, WaterDrag, WaterLift, SubmergedParticle, &WaterSampler,
-			ParticleVol, ShapeVolGeom,
+			ParticleVol, ShapeVolBBox, ShapeVolGeom,
 			ParticleIndex, &ShapeInstances, WaterShapeType, WaterShapeInstance, &ParticleWorldTransform,
 			&SubmergedShapes, &OutTotalSubmergedVol, &OutTotalSubmergedCoM,
 			&OutTotalForce, &OutTotalTorque, &OutTotalBuoyancyForce, &OutTotalBuoyancyTorque]
@@ -777,7 +819,9 @@ namespace BuoyancyAlgorithms
 					Convex = InstancedObject->Object();
 				}				
 
-				// We support convex and boxes separately with the same buoyancy algorithm
+				// We support convex and boxes separately with the same buoyancy algorithm.
+				// Convex uses geometric volume for scaling (submerged vol computed against real geometry).
+				// Box uses bounding box volume (spheres/other shapes are boxified for submersion).
 				if (Convex && !bBuoyancyAlgorithmsUseConvexAsBox)
 				{
 					FBuoyancyConvexShape ConvexShapeQuery(Convex);
@@ -790,13 +834,13 @@ namespace BuoyancyAlgorithms
 						OutSubmergedVol, OutSubmergedCoM, OutForce, OutTorque, OutBuoyancyForce, OutBuoyancyTorque);
 				}
 				else
-				{					
+				{
 					FBuoyancyBoxShape BoxShapeQuery(LocalBox);
 					BoxShapeQuery.Initialize();
 
 					ComputeSubmergedVolumeAndForcesForShape<FBuoyancyBoxShape, SamplerType>(SubmergedParticle, BoxShapeQuery,
 						Evolution, DeltaSeconds, WaterDensity, WaterDrag, WaterLift,
-						ParticleVol, ShapeVolGeom,
+						ParticleVol, ShapeVolBBox,
 						ShapeWorldTransform, WaterSampler,
 						OutSubmergedVol, OutSubmergedCoM, OutForce, OutTorque, OutBuoyancyForce, OutBuoyancyTorque);
 				}
@@ -832,10 +876,12 @@ namespace BuoyancyAlgorithms
 		OutTotalTorque += OutTotalBuoyancyTorque;
 
 		// compute final force and torque for particle as weighted average of all the shapes
-		OutTotalSubmergedCoM /= OutTotalSubmergedVol;
+		if (OutTotalSubmergedVol > UE_SMALL_NUMBER)
+		{
+			OutTotalSubmergedCoM /= OutTotalSubmergedVol;
+		}
 
-		float ScaledSubmergedVol = OutTotalSubmergedVol;
-		ScaleSubmergedVolume(ParticleVol, ShapeVolBBox, ScaledSubmergedVol, OutTotalParticleVol);
+		OutTotalParticleVol = ParticleVol;
 	}
 
 	// forward declaration of usage	
@@ -851,10 +897,16 @@ namespace BuoyancyAlgorithms
 		const Chaos::FPBDRigidsEvolution&, const float, const float, const float, const float,
 		float&, float&, Chaos::FVec3&, Chaos::FVec3&, Chaos::FVec3&);
 
+	template void ComputeSubmergedVolumeAndForcesForParticle<FBuoyancyConstantSplineWithWavesSampler>(FBuoyancyParticleData&,
+		const Chaos::FGeometryParticleHandle*, const Chaos::FGeometryParticleHandle*,
+		TSharedPtr<FBuoyancyConstantSplineWithWavesSampler>,
+		const Chaos::FPBDRigidsEvolution&, const float, const float, const float, const float,
+		float&, float&, Chaos::FVec3&, Chaos::FVec3&, Chaos::FVec3&);
+
 	template <typename ShapeType, typename SamplerType>
 	void ComputeSubmergedVolumeAndForcesForShape(const Chaos::FGeometryParticleHandle* SubmergedParticle, const ShapeType& ShapeQuery,
 		const Chaos::FPBDRigidsEvolution& Evolution, float DeltaSeconds, const float WaterDensity, const float WaterDrag, const float WaterLift,
-		const float ParticleVol, const float ShapeVolGeom,
+		const float ParticleVol, const float ShapeVol,
 		const Chaos::FRigidTransform3 &ShapeWorldTransform, TSharedPtr<SamplerType> WaterSampler,
 		float& OutSubmergedVol, Chaos::FVec3& OutSubmergedCoM,
 		Chaos::FVec3& OutForce, Chaos::FVec3& OutTorque,
@@ -1169,7 +1221,7 @@ namespace BuoyancyAlgorithms
 		// compute buoyancy force and delta velocity for solver
 		Chaos::FVec3 WorldBuoyantForce, WorldBuoyantTorque;
 		BuoyancyAlgorithms::ComputeBuoyantForceForShape(Evolution, RigidParticle, DeltaSeconds,
-			WaterDensity, ParticleVol, ShapeVolGeom,
+			WaterDensity, ParticleVol, ShapeVol,
 			OutSubmergedCoM, OutSubmergedVol, ShapeWaterN, WorldBuoyantForce, WorldBuoyantTorque);
 		
 		// track the buoyancy force separately
@@ -1346,7 +1398,7 @@ namespace BuoyancyAlgorithms
 	}	
 
 	void ComputeBuoyantForceForShape(const Chaos::FPBDRigidsEvolution& Evolution, const Chaos::FPBDRigidParticleHandle* RigidParticle, const float DeltaSeconds, const float WaterDensity,
-		const float ParticleVol, const float ShapeVolGeom,
+		const float ParticleVol, const float ShapeVol,
 		const Chaos::FVec3& SubmergedCoM, const float SubmergedVol, const Chaos::FVec3& WaterN, Chaos::FVec3& OutWorldBuoyantForce, Chaos::FVec3& OutWorldBuoyantTorque)
 	{
 		SCOPE_CYCLE_COUNTER(STAT_BuoyancyAlgorithms_ComputeBuoyantForceForShape)
@@ -1369,7 +1421,7 @@ namespace BuoyancyAlgorithms
 		// Scale submerged volume using precomputed particle/shape volumes
 		float ScaledSubmergedVol = SubmergedVol;
 		float TotalParticleVol;
-		ScaleSubmergedVolume(ParticleVol, ShapeVolGeom, ScaledSubmergedVol, TotalParticleVol);
+		ScaleSubmergedVolume(ParticleVol, ShapeVol, ScaledSubmergedVol, TotalParticleVol);
 
 		// Compute buoyant force
 		//
diff --git a/Engine/Plugins/Experimental/Buoyancy/Source/Runtime/Private/BuoyancySubsystem.cpp b/Engine/Plugins/Experimental/Buoyancy/Source/Runtime/Private/BuoyancySubsystem.cpp
--- a/Engine/Plugins/Experimental/Buoyancy/Source/Runtime/Private/BuoyancySubsystem.cpp
+++ b/Engine/Plugins/Experimental/Buoyancy/Source/Runtime/Private/BuoyancySubsystem.cpp
@@ -63,6 +63,22 @@ static FAutoConsoleVariableRef CVarUseShallowWaterSimulation(
 	ECVF_Scalability
 );
 
+bool bBuoyancyUseWaves = true;
+static FAutoConsoleVariableRef CVarBuoyancyUseWaves(
+	TEXT("p.Buoyancy.bUseWaves"),
+	bBuoyancyUseWaves,
+	TEXT("When true, buoyancy applies Gerstner wave displacement for water bodies that have waves"),
+	ECVF_Scalability
+);
+
+bool bBuoyancyUseSimpleWaves = false;
+static FAutoConsoleVariableRef CVarBuoyancyUseSimpleWaves(
+	TEXT("p.Buoyancy.bUseSimpleWaves"),
+	bBuoyancyUseSimpleWaves,
+	TEXT("When true, use the faster simple wave height evaluation (height only, no normal) instead of the full two-sample steepness evaluation"),
+	ECVF_Scalability
+);
+
 int32 bUseAccurateIntegrationForSplines = 0;
 static FAutoConsoleVariableRef CVarUseAccurateIntegrationForSplines(
 	TEXT("p.Buoyancy.bUseAccurateIntegrationForSplines"),
@@ -71,14 +87,6 @@ static FAutoConsoleVariableRef CVarUseAccurateIntegrationForSplines(
 	ECVF_Scalability
 );
 
-float AccurateIntegrationDragMultiplier = 0.0075;
-static FAutoConsoleVariableRef CVarAccurateIntegrationDragMultiplier(
-	TEXT("p.Buoyancy.AccurateIntegrationDragMultiplier"),
-	AccurateIntegrationDragMultiplier,
-	TEXT("Accurate buoyancy method for splines"),
-	ECVF_Scalability
-);
-
 //
 // Logging
 //
@@ -217,6 +225,9 @@ void UBuoyancySubsystem::PostInitialize()
 {
 	Super::PostInitialize();
 
+	// Cache pointer to ocean fallback depth CVar (defined in Water module)
+	OceanFallbackDepthCVar = IConsoleManager::Get().FindConsoleVariable(TEXT("r.Water.OceanFallbackDepth"));
+
 	// Apply initial runtime settings
 	ApplyRuntimeSettings(GetDefault<UBuoyancyRuntimeSettings>(), EPropertyChangeType::ValueSet);
 
@@ -313,6 +324,9 @@ void UBuoyancySubsystem::Tick(float DeltaTime)
 		UpdateBuoyancySettings();
 	}
 
+	// Send wave reference time to physics thread every frame
+	UpdateWaveReferenceTime();
+
 	if (NetMode != World->GetNetMode())
 	{
 		NetMode = World->GetNetMode();
@@ -331,6 +345,7 @@ void UBuoyancySubsystem::UpdateAllAsyncInputs()
 	UpdateNetMode();
 	UpdateSplineData();
 	UpdateBuoyancySettings();
+	UpdateWaveReferenceTime();
 }
 
 void UBuoyancySubsystem::UpdateNetMode()
@@ -555,13 +570,38 @@ void UBuoyancySubsystem::UpdateSplineData()
 				// child particles and marshaled to PT.
 				const Chaos::FRigidTransform3 WaterTransform = WaterBodyComponent->GetComponentTransform();
 				const TOptional<FInterpCurveFloat> EmptyOptionalFloat;				
+
+				// Build wave data if this water body has Gerstner waves
+				TSharedPtr<FBuoyancyWaveData> WaveData;
+				if (bBuoyancyUseWaves && WaterBodyComponent->HasWaves())
+				{
+					if (UWaterWavesBase* WavesBase = WaterBodyComponent->GetWaterWaves())
+					{
+						if (const UGerstnerWaterWaves* GerstnerWaves = Cast<UGerstnerWaterWaves>(WavesBase->GetWaterWaves()))
+						{
+							WaveData = MakeShared<FBuoyancyGerstnerWaveData>(
+								GerstnerWaves->GetGerstnerWaves(),
+								GerstnerWaves->GetMaxWaveHeight(),
+								WaterBodyComponent->TargetWaveMaskDepth);
+						}
+					}
+				}
+
+				const float OceanFallbackDepth = OceanFallbackDepthCVar ? OceanFallbackDepthCVar->GetFloat() : 0.f;
+
 				TSharedPtr<FBuoyancyWaterSplineData> WaterSplineData = MakeShared<FBuoyancyWaterSplineData>(
 					WaterTransform,
 					SplineComponent->GetSplinePointsPosition(),
 					WaterBodyComponent->GetWaterBodyType(),
 					SplineMetadata ? SplineMetadata->RiverWidth : EmptyOptionalFloat,
 					SplineMetadata ? SplineMetadata->WaterVelocityScalar : EmptyOptionalFloat,
-					bUseShallowWaterSimulation && WaterBodyComponent->UseBakedSimulationForQueriesAndPhysics() ? BakedSim->BakedSimulationData : nullptr
+					bUseShallowWaterSimulation && WaterBodyComponent->UseBakedSimulationForQueriesAndPhysics() ? BakedSim->BakedSimulationData : nullptr,
+					WaveData,
+					SplineMetadata ? SplineMetadata->Depth : EmptyOptionalFloat,
+					WaterBodyComponent->GetChannelDepth(),
+					WaterBodyComponent->CurveSettings.CurveRampWidth,
+					WaterBodyComponent->CurveSettings.ChannelEdgeOffset,
+					OceanFallbackDepth
 				);
 
 				// Go over each physics object in each primitive component which was generated
@@ -601,6 +641,22 @@ void UBuoyancySubsystem::UpdateBuoyancySettings()
 	}
 }
 
+void UBuoyancySubsystem::UpdateWaveReferenceTime()
+{
+	if (SimCallback)
+	{
+		if (FBuoyancySubsystemSimCallbackInput* AsyncInput = SimCallback->GetProducerInputData_External())
+		{
+			float Time = 0.f;
+			if (UWaterSubsystem* WaterSubsystem = UWaterSubsystem::GetWaterSubsystem(GetWorld()))
+			{
+				Time = WaterSubsystem->GetWaterTimeSeconds();
+			}
+			AsyncInput->WaveReferenceTime = Time;
+		}
+	}
+}
+
 void UBuoyancySubsystem::ProcessSurfaceTouchCallbacks()
 {
 	SCOPE_CYCLE_COUNTER(STAT_BuoyancySubsystem_DispatchCallbacks)
@@ -736,6 +792,7 @@ void FBuoyancySubsystemSimCallbackInput::Reset()
 	SplineData.Reset();
 	BuoyancySettings.Reset();
 	NetMode.Reset();
+	WaveReferenceTime.Reset();
 }
 
 void FBuoyancySubsystemSimCallbackOutput::Reset()
@@ -801,6 +858,11 @@ void FBuoyancySubsystemSimCallback::OnPreSimulate_Internal()
 				SplineKeyCache.SetCacheLimit(BuoyancySettings->SplineKeyCacheLimit);
 			}
 		}
+
+		if (Input->WaveReferenceTime.IsSet())
+		{
+			WaveReferenceTime = *Input->WaveReferenceTime;
+		}
 	}
 
 	// If we don't have a valid buoyancy settings object, don't continue
@@ -1050,6 +1112,18 @@ bool FBuoyancySubsystemSimCallback::QuerySpline(
 				if (DiffProj < SMALL_NUMBER) { return false; }
 				break;
 			}
+
+			case EWaterBodyType::Ocean:
+			{
+				// Determine if we're inside the ocean by projecting the horizontal spline
+				// diff onto the cross product of the spline direction and the up-vector
+				// (ie, the right-vector).  Note this is reversed from lakes since the splines
+				// define the carved out area that is land.
+				const FVector RightVector = FVector::CrossProduct(ClosestPointDerivative, FVector::UpVector);
+				const float DiffProj = FVector::DotProduct(RightVector, LateralDiff);
+				if (DiffProj > SMALL_NUMBER) { return false; }
+				break;
+			}
 		}
 			
 #if CHAOS_DEBUG_DRAW
@@ -1147,8 +1221,32 @@ void FBuoyancySubsystemSimCallback::TrackInteraction(
 			return;
 		}		
 		
+		if (bBuoyancyUseWaves && WaterSpline.HasWaves())
+		{
+			WaterSampler = TSharedPtr<FBuoyancyWaterSampler>(new FBuoyancyConstantSplineWithWavesSampler(
+				WaterSpline, ParticlePos, ClosestPoint, ClosestPointDerivative, ClosestSplineKey, WaterN, WaveReferenceTime, bBuoyancyUseSimpleWaves));
+
+#if CHAOS_DEBUG_DRAW
+			if (bBuoyancyDebugDraw)
+			{
+				auto DerivedSampler = StaticCastSharedPtr<FBuoyancyConstantSplineWithWavesSampler>(WaterSampler);
+
+				Chaos::FDebugDrawQueue::GetInstance().DrawDebugDirectionalArrow(ParticlePos, ClosestPoint, 20, FColor::Purple, false, -1.f, -1, 2.f);
+				
+				FVector BottomPoint = ParticlePos;
+				FVector Vel;
+				float H, D;
+				DerivedSampler->SampleWaterAtPosition(ParticlePos, Vel, H, D);
+				BottomPoint.Z = ClosestPoint.Z - D;
+				Chaos::FDebugDrawQueue::GetInstance().DrawDebugDirectionalArrow(ParticlePos, BottomPoint, 20, FColor::Purple, false, -1.f, -1, 2.f);
+			}
+#endif
+		}
+		else
+		{
 		WaterSampler = TSharedPtr<FBuoyancyWaterSampler>(new FBuoyancyConstantSplineSampler(WaterSpline, ClosestPoint, ClosestPointDerivative, ClosestSplineKey, WaterN));				
 	}
+	}
 
 	{
 		SCOPE_CYCLE_COUNTER(STAT_BuoyancySubsystem_AddInteraction)
@@ -1183,7 +1281,9 @@ void FBuoyancySubsystemSimCallback::ProcessInteractions(Chaos::FPBDRigidsEvoluti
 				UE_LOG(LogBuoyancySubsystem, Warning, TEXT("Skipped invalid buoyancy interaction"));
 				return;
 			}
-			else if (Interaction.WaterSampler->GetSamplerType() == EWaterSamplerType::ShallowWater || bUseAccurateIntegrationForSplines)
+			else if (Interaction.WaterSampler->GetSamplerType() == EWaterSamplerType::ShallowWater
+				|| Interaction.WaterSampler->GetSamplerType() == EWaterSamplerType::ConstantSplineWithWaves
+				|| bUseAccurateIntegrationForSplines)
 			{
 				ProcessAccurateInteraction(Evolution, Interaction, Interaction.WaterSampler);
 			}
@@ -1205,9 +1305,6 @@ void FBuoyancySubsystemSimCallback::ProcessAccurateInteraction(Chaos::FPBDRigids
 	Chaos::FVec3 WorldTorque = Chaos::FVec3::ZeroVector;
 	float TotalParticleVol = 0;
 	
-	const float AccurateIntegrationDrag = AccurateIntegrationDragMultiplier * BuoyancySettings->WaterDrag;
-	const float AccurateIntegrationLift = AccurateIntegrationDragMultiplier * BuoyancySettings->WaterLift;
-
 	if (!Interaction.WaterSampler.IsValid())
 	{
 		UE_LOG(LogBuoyancySubsystem, Warning, TEXT("Skipped invalid buoyancy interaction "));
@@ -1218,7 +1315,15 @@ void FBuoyancySubsystemSimCallback::ProcessAccurateInteraction(Chaos::FPBDRigids
 		BuoyancyAlgorithms::ComputeSubmergedVolumeAndForcesForParticle(BuoyancyParticleData,
 			Interaction.RigidParticle, Interaction.WaterParticle,
 			StaticCastSharedPtr<FBuoyancyConstantSplineSampler>(Interaction.WaterSampler),
-			Evolution, DeltaSeconds, BuoyancySettings->WaterDensity, AccurateIntegrationDrag, AccurateIntegrationLift,
+			Evolution, DeltaSeconds, BuoyancySettings->WaterDensity, BuoyancySettings->WaterDrag, BuoyancySettings->WaterLift,
+			TotalParticleVol, SubmergedVol, SubmergedCoM, WorldForce, WorldTorque);
+	}
+	else if (Interaction.WaterSampler->GetSamplerType() == EWaterSamplerType::ConstantSplineWithWaves)
+	{
+		BuoyancyAlgorithms::ComputeSubmergedVolumeAndForcesForParticle(BuoyancyParticleData,
+			Interaction.RigidParticle, Interaction.WaterParticle,
+			StaticCastSharedPtr<FBuoyancyConstantSplineWithWavesSampler>(Interaction.WaterSampler),
+			Evolution, DeltaSeconds, BuoyancySettings->WaterDensity, BuoyancySettings->WaterDrag, BuoyancySettings->WaterLift,
 			TotalParticleVol, SubmergedVol, SubmergedCoM, WorldForce, WorldTorque);		
 	}
 	else if (Interaction.WaterSampler->GetSamplerType() == EWaterSamplerType::ShallowWater)
@@ -1226,7 +1331,7 @@ void FBuoyancySubsystemSimCallback::ProcessAccurateInteraction(Chaos::FPBDRigids
 		BuoyancyAlgorithms::ComputeSubmergedVolumeAndForcesForParticle(BuoyancyParticleData,
 			Interaction.RigidParticle, Interaction.WaterParticle,
 			StaticCastSharedPtr<FBuoyancyShallowWaterSampler>(Interaction.WaterSampler),
-			Evolution, DeltaSeconds, BuoyancySettings->WaterDensity, AccurateIntegrationDrag, AccurateIntegrationLift,
+			Evolution, DeltaSeconds, BuoyancySettings->WaterDensity, BuoyancySettings->WaterDrag, BuoyancySettings->WaterLift,
 			TotalParticleVol, SubmergedVol, SubmergedCoM, WorldForce, WorldTorque);
 	}
 
diff --git a/Engine/Plugins/Experimental/Buoyancy/Source/Runtime/Private/BuoyancyWaveData.cpp b/Engine/Plugins/Experimental/Buoyancy/Source/Runtime/Private/BuoyancyWaveData.cpp
new file mode 100644
--- /dev/null
+++ b/Engine/Plugins/Experimental/Buoyancy/Source/Runtime/Private/BuoyancyWaveData.cpp
@@ -0,0 +1,49 @@
+// Copyright Epic Games, Inc. All Rights Reserved.
+
+#include "BuoyancyWaveData.h"
+#include "GerstnerWaveEvaluation.h"
+#include "Misc/LargeWorldRenderPosition.h"
+
+// ----------------------------------------------------------------------------------
+// FBuoyancyGerstnerWaveData
+// ----------------------------------------------------------------------------------
+
+FBuoyancyGerstnerWaveData::FBuoyancyGerstnerWaveData(
+	const TArray<FGerstnerWave>& InWaveParams,
+	float InMaxWaveHeight,
+	float InTargetWaveMaskDepth)
+	: WaveParams(InWaveParams)
+	, MaxWaveHeightValue(InMaxWaveHeight)
+	, TargetWaveMaskDepth(InTargetWaveMaskDepth)
+{
+}
+
+float FBuoyancyGerstnerWaveData::GetMaxWaveHeight() const
+{
+	return MaxWaveHeightValue;
+}
+
+float FBuoyancyGerstnerWaveData::GetSimpleWaveHeightAtPosition(const FVector& InPosition, float InTime) const
+{
+	// Use the offset of the normalized tile as world position to match the shader behavior (see GerstnerWaveFunctions.ush).
+	FVector WorldPosition(FLargeWorldRenderPosition(InPosition).GetOffset());
+	return GerstnerWaveEvaluation::GetSimpleWaveHeightAtPosition(WaveParams, WorldPosition, InTime, /*bBlendLWCTiles=*/true);
+}
+
+float FBuoyancyGerstnerWaveData::GetWaveHeightAtPosition(const FVector& InPosition, float InTime, FVector& OutNormal) const
+{
+	// Use the offset of the normalized tile as world position to match the shader behavior (see GerstnerWaveFunctions.ush).
+	FVector WorldPosition(FLargeWorldRenderPosition(InPosition).GetOffset());
+	return GerstnerWaveEvaluation::GetWaveHeightAtPosition(WaveParams, WorldPosition, InTime, OutNormal, /*bBlendLWCTiles=*/true);
+}
+
+float FBuoyancyGerstnerWaveData::GetWaveAttenuationFactor(const FVector& Position, float InWaterDepth) const
+{
+	// Guard against TargetWaveMaskDepth=0 which would cause division by zero (NaN) in the
+	// exponential. A zero mask depth means no attenuation is configured, so return full strength.
+	if (TargetWaveMaskDepth <= UE_SMALL_NUMBER)
+	{
+		return 1.f;
+	}
+	return GerstnerWaveEvaluation::GetWaveAttenuationFactor(InWaterDepth, TargetWaveMaskDepth);
+}
diff --git a/Engine/Plugins/Experimental/Buoyancy/Source/Runtime/Public/BuoyancyAlgorithms.h b/Engine/Plugins/Experimental/Buoyancy/Source/Runtime/Public/BuoyancyAlgorithms.h
--- a/Engine/Plugins/Experimental/Buoyancy/Source/Runtime/Public/BuoyancyAlgorithms.h
+++ b/Engine/Plugins/Experimental/Buoyancy/Source/Runtime/Public/BuoyancyAlgorithms.h
@@ -18,6 +18,7 @@ namespace Chaos
 class FBuoyancyWaterSampler;
 class FBuoyancyConstantSplineSampler;
 class FBuoyancyShallowWaterSampler;
+class FBuoyancyConstantSplineWithWavesSampler;
 
 struct FBuoyancyParticleData;
 
@@ -146,6 +147,9 @@ namespace BuoyancyAlgorithms
 	// maximum possible output value of the non-scaled ComputeSubmergedVolume.
 	Chaos::FRealSingle ComputeShapeVolume(const Chaos::FGeometryParticleHandle* Particle, const bool UseBoundingBoxVolume = true);
 
+	// Compute both bounding box and geometric shape volumes in a single traversal.
+	void ComputeShapeVolumes(const Chaos::FGeometryParticleHandle* Particle, Chaos::FRealSingle& OutBBoxVol, Chaos::FRealSingle& OutGeomVol);
+
 	//
 	void ScaleSubmergedVolume(const Chaos::FPBDRigidsEvolutionGBF& Evolution, const Chaos::FGeometryParticleHandle* Particle, const bool UseBoundingBoxVolume, Chaos::FRealSingle& SubmergedVol, Chaos::FRealSingle& TotalVol);
 
@@ -190,7 +194,7 @@ namespace BuoyancyAlgorithms
 	void ComputeSubmergedVolumeAndForcesForShape(
 		const Chaos::FGeometryParticleHandle* SubmergedParticle, const ShapeType& BoxShape,
 		const Chaos::FPBDRigidsEvolution& Evolution, float DeltaSeconds, const float WaterDensity, const float WaterDrag, const float WaterLift,
-		const float ParticleVol, const float ShapeVolGeom,
+		const float ParticleVol, const float ShapeVol,
 		const Chaos::FRigidTransform3 &ShapeWorldTransform, TSharedPtr<SamplerType> WaterSampler,
 		float& OutSubmergedVol, Chaos::FVec3& OutSubmergedCoM,
 		Chaos::FVec3& OutForce, Chaos::FVec3& OutTorque,
@@ -300,6 +304,6 @@ namespace BuoyancyAlgorithms
 
 	// compute the force the buoyancy exerts on a shape
 	void ComputeBuoyantForceForShape(const Chaos::FPBDRigidsEvolution& Evolution, const Chaos::FPBDRigidParticleHandle* RigidParticle, const float DeltaSeconds, const float WaterDensity,
-		const float ParticleVol, const float ShapeVolGeom,
+		const float ParticleVol, const float ShapeVol,
 		const Chaos::FVec3& SubmergedCoM, const float SubmergedVol, const Chaos::FVec3& WaterN, Chaos::FVec3& OutWorldBuoyantForce, Chaos::FVec3& OutWorldBuoyantTorque);
 }
diff --git a/Engine/Plugins/Experimental/Buoyancy/Source/Runtime/Public/BuoyancyParticleData.h b/Engine/Plugins/Experimental/Buoyancy/Source/Runtime/Public/BuoyancyParticleData.h
--- a/Engine/Plugins/Experimental/Buoyancy/Source/Runtime/Public/BuoyancyParticleData.h
+++ b/Engine/Plugins/Experimental/Buoyancy/Source/Runtime/Public/BuoyancyParticleData.h
@@ -17,6 +17,7 @@ enum class EWaterSamplerType : uint8
 {
 	ConstantSpline,
 	ShallowWater,
+	ConstantSplineWithWaves,
 	Invalid,
 };
 
@@ -160,6 +161,96 @@ public:
 	bool WaterVelSet = false;
 };
 
+// Sampler for water bodies with Gerstner waves. Evaluates wave height and normal
+// per-position using the wave data stored on the spline data struct.
+class FBuoyancyConstantSplineWithWavesSampler : public FBuoyancyWaterSampler
+{
+public:
+	FBuoyancyConstantSplineWithWavesSampler(
+		const FBuoyancyWaterSplineData& InWaterSpline,
+		const FVector& InParticlePosition,
+		const FVector& InClosestPoint,
+		const FVector& InClosestPointDerivative,
+		const float InClosestSplineKey,
+		const FVector& InWaterN,
+		float InWaveReferenceTime,
+		bool bInUseSimpleWaves = false)
+		: WaterSpline(InWaterSpline)
+		, SplineBaseHeight(InClosestPoint.Z)
+		, ClosestPointDerivative(InClosestPointDerivative)
+		, ClosestSplineKey(InClosestSplineKey)
+		, BaseWaterN(InWaterN)
+		, WaveReferenceTime(InWaveReferenceTime)
+		, bUseSimpleWaves(bInUseSimpleWaves)
+	{
+		SamplerType = EWaterSamplerType::ConstantSplineWithWaves;
+		bUsePerShapeWaterPlane = true;
+		ClosestWaterPointToParticle = InClosestPoint;
+
+		// Compute velocity from spline (same as ConstantSpline sampler)
+		AverageVelocity = WaterSpline.Velocity.IsSet()
+			? WaterSpline.Velocity->Eval(ClosestSplineKey) * ClosestPointDerivative.GetSafeNormal()
+			: FVector::ZeroVector;
+
+		// Estimate water depth at the particle's position.
+		EstimatedDepth = InWaterSpline.EstimateWaterDepth(InParticlePosition, InClosestPoint, InClosestSplineKey);
+		WaveAttenuationFactor = InWaterSpline.WaveData.IsValid()
+			? InWaterSpline.WaveData->GetWaveAttenuationFactor(InParticlePosition, EstimatedDepth)
+			: 1.f;
+	}
+
+	virtual void SampleWaterAtPosition(const FVector& WorldPosition, FVector& OutVelocity, float& OutHeight, float& OutDepth) const override
+	{
+		OutVelocity = AverageVelocity;
+
+		// Base spline height + wave perturbation at query position, attenuated by depth
+		float WaveHeight = 0.f;
+		if (WaterSpline.WaveData.IsValid())
+		{
+			if (bUseSimpleWaves)
+			{
+				WaveHeight = WaterSpline.WaveData->GetSimpleWaveHeightAtPosition(WorldPosition, WaveReferenceTime);
+			}
+			else
+			{
+				FVector WaveNormal;
+				WaveHeight = WaterSpline.WaveData->GetWaveHeightAtPosition(WorldPosition, WaveReferenceTime, WaveNormal);
+			}
+		}
+		OutHeight = SplineBaseHeight + WaveHeight * WaveAttenuationFactor;
+
+		OutDepth = EstimatedDepth;
+	}
+
+	virtual FVector SampleVelocityAtPosition(const FVector& WorldPosition) const override
+	{
+		return AverageVelocity;
+	}
+
+	virtual FVector SampleNormalAtPosition(const FVector& WorldPosition) const override
+	{
+		if (!bUseSimpleWaves && WaterSpline.WaveData.IsValid())
+		{
+			FVector WaveNormal;
+			WaterSpline.WaveData->GetWaveHeightAtPosition(WorldPosition, WaveReferenceTime, WaveNormal);
+			WaveNormal = FMath::Lerp(BaseWaterN, WaveNormal, WaveAttenuationFactor);
+			return WaveNormal.IsZero() ? BaseWaterN : WaveNormal.GetSafeNormal();
+		}
+		return BaseWaterN;
+	}
+
+public:
+	const FBuoyancyWaterSplineData& WaterSpline;
+	float SplineBaseHeight = 0.f;
+	FVector ClosestPointDerivative = FVector::ZeroVector;
+	float ClosestSplineKey = 0.f;
+	FVector BaseWaterN = FVector::ZeroVector;
+	float WaveReferenceTime = 0.f;
+	bool bUseSimpleWaves = false;
+	float EstimatedDepth = MAX_FLT;
+	float WaveAttenuationFactor = 1.f;
+};
+
 // Each particle will have a list of potential midphases to process,
 // which must be sorted in descending Z order. This struct is used
 // to store them
diff --git a/Engine/Plugins/Experimental/Buoyancy/Source/Runtime/Public/BuoyancySubsystem.h b/Engine/Plugins/Experimental/Buoyancy/Source/Runtime/Public/BuoyancySubsystem.h
--- a/Engine/Plugins/Experimental/Buoyancy/Source/Runtime/Public/BuoyancySubsystem.h
+++ b/Engine/Plugins/Experimental/Buoyancy/Source/Runtime/Public/BuoyancySubsystem.h
@@ -52,9 +52,9 @@ struct FBuoyancySettings
 	// Source: https://en.wikipedia.org/wiki/Properties_of_water
 	float WaterDensity = 0.0001f; // kg/cm^3
 
-	float MaxDeltaV = 200.f; // cm/s
+	float MaxDeltaV = 500.f; // cm/s
 
-	float MaxDeltaW = 2.f; // rad/s
+	float MaxDeltaW = 5.f; // rad/s
 
 	float WaterDrag = 1.f; // unitless
 
@@ -203,6 +203,9 @@ private:
 	// Put updated settings struct onto async input to be sent to sim callback
 	void UpdateBuoyancySettings();
 
+	// Send wave reference time to physics thread each frame
+	void UpdateWaveReferenceTime();
+
 	// Process async outputs which hold data for triggering callbacks
 	void ProcessSurfaceTouchCallbacks();
 
@@ -215,6 +218,9 @@ private:
 
 	bool bBuoyancySettingsChanged;
 
+	// Cached pointer to r.Water.OceanFallbackDepth CVar (looked up once in PostInitialize)
+	IConsoleVariable* OceanFallbackDepthCVar = nullptr;
+
 	FBuoyancySettings BuoyancySettings;
 
 	FBuoyancyWaterSplineDataManager* SplineData;
@@ -248,6 +254,9 @@ struct FBuoyancySubsystemSimCallbackInput : public Chaos::FSimCallbackInput
 	// Set when net mode changes - should be one time on initialization.
 	TOptional<ENetMode> NetMode;
 
+	// Wave reference time, sent every frame when waves are active
+	TOptional<float> WaveReferenceTime;
+
 	void Reset();
 };
 
@@ -322,6 +331,9 @@ private:
 	// sim tick is taking place in.
 	ENetMode NetMode = ENetMode::NM_MAX;
 
+	// Wave reference time for evaluating wave heights on the physics thread
+	float WaveReferenceTime = 0.0f;
+
 	// Local cache of spline keys to reduce spline evaluations.
 	FSplineKeyCacheGrid SplineKeyCache;
 };
diff --git a/Engine/Plugins/Experimental/Buoyancy/Source/Runtime/Public/BuoyancyWaterSplineData.h b/Engine/Plugins/Experimental/Buoyancy/Source/Runtime/Public/BuoyancyWaterSplineData.h
--- a/Engine/Plugins/Experimental/Buoyancy/Source/Runtime/Public/BuoyancyWaterSplineData.h
+++ b/Engine/Plugins/Experimental/Buoyancy/Source/Runtime/Public/BuoyancyWaterSplineData.h
@@ -3,6 +3,7 @@
 #pragma once
 
 #include "WaterBodyComponent.h"
+#include "BuoyancyWaveData.h"
 
 struct FBuoyancyWaterSplineData
 {
@@ -13,7 +14,13 @@ struct FBuoyancyWaterSplineData
 		const EWaterBodyType InBodyType,
 		const TOptional<FInterpCurveFloat>& InWidth,
 		const TOptional<FInterpCurveFloat>& InVelocity,
-		const TWeakObjectPtr<UShallowWaterSimulationDataBase> InShallowWaterSimData
+		const TWeakObjectPtr<UShallowWaterSimulationDataBase> InShallowWaterSimData,
+		TSharedPtr<FBuoyancyWaveData> InWaveData = nullptr,
+		const TOptional<FInterpCurveFloat>& InDepth = TOptional<FInterpCurveFloat>(),
+		float InChannelDepth = 0.f,
+		float InCurveRampWidth = 512.f,
+		float InChannelEdgeOffset = 0.f,
+		float InOceanFallbackDepth = 0.f
 	)
 		: Transform(InTransform)
 		, Position(InPosition)
@@ -21,6 +28,12 @@ struct FBuoyancyWaterSplineData
 		, Width(InWidth)
 		, Velocity(InVelocity)
 		, ShallowWaterSimData(InShallowWaterSimData)
+		, WaveData(InWaveData)
+		, Depth(InDepth)
+		, ChannelDepth(InChannelDepth)
+		, CurveRampWidth(InCurveRampWidth)
+		, ChannelEdgeOffset(InChannelEdgeOffset)
+		, OceanFallbackDepth(InOceanFallbackDepth)
 	{ }
 
 
@@ -41,4 +54,69 @@ struct FBuoyancyWaterSplineData
 	// Only some water bodies can have a baked shallow water sim representation, which
 	// overrides the splines for water height/depth/velocity/normal
 	const TWeakObjectPtr<UShallowWaterSimulationDataBase> ShallowWaterSimData;
+
+	// Wave data (null if water body has no waves)
+	TSharedPtr<FBuoyancyWaveData> WaveData;
+
+	// For rivers: actual depth curve from spline metadata
+	TOptional<FInterpCurveFloat> Depth;
+
+	// For ocean/lake: estimate depth from distance to shoreline spline
+	float ChannelDepth = 0.f;
+	float CurveRampWidth = 512.f;
+	float ChannelEdgeOffset = 0.f;
+
+	// For ocean: fallback depth when beyond terrain carving zone
+	float OceanFallbackDepth = 0.f;
+
+	bool HasWaves() const { return WaveData.IsValid(); }
+
+	// Estimate water depth at a world position given closest spline point and key.
+	// Rivers use the actual depth spline curve; ocean/lake estimate from XY distance to shoreline.
+	float EstimateWaterDepth(const FVector& WorldPosition, const FVector& ClosestSplinePoint, float SplineKey) const
+	{
+		if (BodyType == EWaterBodyType::River && Depth.IsSet())
+		{
+			// Rivers: use actual depth curve from spline metadata
+			return FMath::Max(Depth->Eval(SplineKey, 0.f), 0.f);
+		}
+
+		// Ocean/Lake: estimate from XY distance to shoreline spline.
+		// Terrain carving starts at ChannelEdgeOffset from the spline and
+		// ramps to ChannelDepth over CurveRampWidth.
+		const float DistToSpline = FVector::Dist2D(WorldPosition, ClosestSplinePoint);
+		const float EffectiveDist = FMath::Max(DistToSpline - ChannelEdgeOffset, 0.f);
+		const float SafeRampWidth = FMath::Max(CurveRampWidth, 1.f);
+
+		if (BodyType == EWaterBodyType::Ocean)
+		{
+			// Ocean: depth ramps from 0 at the carving edge to ChannelDepth,
+			// then continues toward OceanFallbackDepth for deep water.
+			const float FallbackDepth = FMath::Max(OceanFallbackDepth, ChannelDepth);
+			if (FallbackDepth <= 0.f)
+			{
+				return MAX_FLT;
+			}
+			if (ChannelDepth > 0.f)
+			{
+				// Within carving ramp: linear ramp to ChannelDepth
+				if (EffectiveDist <= SafeRampWidth)
+				{
+					return ChannelDepth * (EffectiveDist / SafeRampWidth);
+				}
+				// Beyond carving: transition from ChannelDepth to FallbackDepth
+				const float T = FMath::Clamp((EffectiveDist - SafeRampWidth) / SafeRampWidth, 0.f, 1.f);
+				return FMath::Lerp(ChannelDepth, FallbackDepth, T);
+			}
+			// No carving configured: ramp to fallback depth
+			return FallbackDepth * FMath::Clamp(EffectiveDist / SafeRampWidth, 0.f, 1.f);
+		}
+
+		// Lake/other: depth caps at ChannelDepth (lake bottom = carved terrain)
+		if (ChannelDepth <= 0.f)
+		{
+			return MAX_FLT;
+		}
+		return ChannelDepth * FMath::Clamp(EffectiveDist / SafeRampWidth, 0.f, 1.f);
+	}
 };
\ No newline at end of file
diff --git a/Engine/Plugins/Experimental/Buoyancy/Source/Runtime/Public/BuoyancyWaveData.h b/Engine/Plugins/Experimental/Buoyancy/Source/Runtime/Public/BuoyancyWaveData.h
new file mode 100644
--- /dev/null
+++ b/Engine/Plugins/Experimental/Buoyancy/Source/Runtime/Public/BuoyancyWaveData.h
@@ -0,0 +1,50 @@
+// Copyright Epic Games, Inc. All Rights Reserved.
+
+#pragma once
+
+#include "CoreMinimal.h"
+#include "GerstnerWaterWaves.h"
+
+// Abstract base for wave data evaluable on the physics thread (no UObject access).
+// Copied from game-thread UWaterWaves data so that the physics thread can evaluate
+// wave heights without touching UObjects.
+struct FBuoyancyWaveData
+{
+	virtual ~FBuoyancyWaveData() = default;
+
+	virtual float GetMaxWaveHeight() const = 0;
+
+	// Fast version: height only, no normal perturbation
+	virtual float GetSimpleWaveHeightAtPosition(
+		const FVector& InPosition, float InTime) const = 0;
+
+	// Full version: height + perturbed normal
+	virtual float GetWaveHeightAtPosition(
+		const FVector& InPosition, float InTime, FVector& OutNormal) const = 0;
+
+	// Attenuation factor (0-1) based on water depth (e.g. shore dampening)
+	virtual float GetWaveAttenuationFactor(
+		const FVector& InPosition, float InWaterDepth) const = 0;
+};
+
+// Gerstner wave implementation - self-contained copy of wave params for physics thread use.
+// Math mirrors UGerstnerWaterWaves::GetWaveHeightAtPosition / GetSimpleWaveHeightAtPosition.
+struct FBuoyancyGerstnerWaveData : public FBuoyancyWaveData
+{
+	TArray<FGerstnerWave> WaveParams;
+	float MaxWaveHeightValue = 0.0f;
+	float TargetWaveMaskDepth = 0.0f;
+
+	FBuoyancyGerstnerWaveData(
+		const TArray<FGerstnerWave>& InWaveParams,
+		float InMaxWaveHeight,
+		float InTargetWaveMaskDepth);
+
+	virtual float GetMaxWaveHeight() const override;
+	virtual float GetSimpleWaveHeightAtPosition(
+		const FVector& InPosition, float InTime) const override;
+	virtual float GetWaveHeightAtPosition(
+		const FVector& InPosition, float InTime, FVector& OutNormal) const override;
+	virtual float GetWaveAttenuationFactor(
+		const FVector& InPosition, float InWaterDepth) const override;
+};
